Fix for build break caused by CL#2070165

diff --git a/Engine/Source/Programs/UnrealPak/Private/KeyGenerator.cpp b/Engine/Source/Programs/UnrealPak/Private/KeyGenerator.cpp
--- a/Engine/Source/Programs/UnrealPak/Private/KeyGenerator.cpp
+++ b/Engine/Source/Programs/UnrealPak/Private/KeyGenerator.cpp
@@ -43,7 +43,7 @@ public:
 	{
 		// Must be odd number
 		check(!(Candidate & One).IsZero());
-		Thread = FRunnableThread::Create(this, TEXT("FPrimeCheckRunnable"), false, false, 0, TPri_Normal);
+		Thread = FRunnableThread::Create(this, TEXT("FPrimeCheckRunnable"));
 	}
 
 	virtual ~FPrimeCheckRunnable()
