Fix for build break caused  by CL#2070165.

diff --git a/Engine/Source/Programs/UnrealPak/Private/KeyGenerator.cpp b/Engine/Source/Programs/UnrealPak/Private/KeyGenerator.cpp
--- a/Engine/Source/Programs/UnrealPak/Private/KeyGenerator.cpp
+++ b/Engine/Source/Programs/UnrealPak/Private/KeyGenerator.cpp
@@ -300,7 +300,7 @@ public:
 	{
 		// Must be an odd number
 		check(!(MinValue & One).IsZero());
-		Thread = FRunnableThread::Create(this, TEXT("FPrimeFinderRunnable"), false, false, 0, TPri_Normal);
+		Thread = FRunnableThread::Create(this, TEXT("FPrimeFinderRunnable"));
 	}
 
 	virtual ~FPrimeFinderRunnable()
