Adding encryption and compression support to UnrealPak & Pak files.

diff --git a/Engine/Source/Programs/AutomationTool/Scripts/CookCommand.Automation.cs b/Engine/Source/Programs/AutomationTool/Scripts/CookCommand.Automation.cs
--- a/Engine/Source/Programs/AutomationTool/Scripts/CookCommand.Automation.cs
+++ b/Engine/Source/Programs/AutomationTool/Scripts/CookCommand.Automation.cs
@@ -71,10 +71,6 @@ public partial class Project : CommandUtils
 		try
 		{
 			var CommandletParams = "-buildmachine -Unversioned -fileopenlog";
-			if (Params.Compressed)
-			{
-				CommandletParams += " -Compressed";
-			}
             if (Params.UseDebugParamForEditorExe)
             {
                 CommandletParams += " -debug";
diff --git a/Engine/Source/Programs/AutomationTool/Scripts/CopyBuildToStagingDirectory.Automation.cs b/Engine/Source/Programs/AutomationTool/Scripts/CopyBuildToStagingDirectory.Automation.cs
--- a/Engine/Source/Programs/AutomationTool/Scripts/CopyBuildToStagingDirectory.Automation.cs
+++ b/Engine/Source/Programs/AutomationTool/Scripts/CopyBuildToStagingDirectory.Automation.cs
@@ -25,13 +25,18 @@ public partial class Project : CommandUtils
 	/// </summary>
 	/// <param name="Filename"></param>
 	/// <param name="ResponseFile"></param>
-	private static void WritePakResponseFile(string Filename, Dictionary<string, string> ResponseFile)
+	private static void WritePakResponseFile(string Filename, Dictionary<string, string> ResponseFile, bool Compressed)
 	{
         using (var Writer = new StreamWriter(Filename, false, new System.Text.UTF8Encoding(true)))
         {
             foreach (var Entry in ResponseFile)
             {
-                Writer.WriteLine("\"{0}\" \"{1}\"", Entry.Key, Entry.Value);
+                string Line = String.Format("\"{0}\" \"{1}\"", Entry.Key, Entry.Value);
+                if (Compressed)
+                {
+                    Line += " -compress";
+                }
+                Writer.WriteLine(Line);
             }
         }
 	}
@@ -48,7 +53,7 @@ public partial class Project : CommandUtils
 		return Result;
 	}
 
-	static public void RunUnrealPak(Dictionary<string, string> UnrealPakResponseFile, string OutputLocation, string EncryptionKeys, string PakOrderFileLocation)
+	static public void RunUnrealPak(Dictionary<string, string> UnrealPakResponseFile, string OutputLocation, string EncryptionKeys, string PakOrderFileLocation, bool Compressed)
 	{
 		if (UnrealPakResponseFile.Count < 1)
 		{
@@ -56,7 +61,7 @@ public partial class Project : CommandUtils
 		}
 		string PakName = Path.GetFileNameWithoutExtension(OutputLocation);
 		string UnrealPakResponseFileName = CombinePaths(CmdEnv.LogFolder, "PakList_" + PakName + ".txt");
-		WritePakResponseFile(UnrealPakResponseFileName, UnrealPakResponseFile);
+        WritePakResponseFile(UnrealPakResponseFileName, UnrealPakResponseFile, Compressed);
 
 		var UnrealPakExe = CombinePaths(CmdEnv.LocalRoot, "Engine/Binaries/Win64/unrealpak.exe");
 
@@ -399,7 +404,7 @@ public partial class Project : CommandUtils
 			PakOrderFileLocation = CombinePaths(PakOrderFileLocationBase, "EditorOpenOrder.log");
 		}
 
-		RunUnrealPak(UnrealPakResponseFile, OutputLocation, Params.SignPak, PakOrderFileLocation);
+		RunUnrealPak(UnrealPakResponseFile, OutputLocation, Params.SignPak, PakOrderFileLocation, Params.Compressed);
 
 		// add the pak file as needing deployment and convert to lower case again if needed
 		SC.UFSStagingFiles.Add(OutputLocation, OutputRealtiveLocation);		
diff --git a/Engine/Source/Programs/UnrealPak/Private/UnrealPak.cpp b/Engine/Source/Programs/UnrealPak/Private/UnrealPak.cpp
--- a/Engine/Source/Programs/UnrealPak/Private/UnrealPak.cpp
+++ b/Engine/Source/Programs/UnrealPak/Private/UnrealPak.cpp
@@ -7,6 +7,7 @@
 #include "BigInt.h"
 #include "SignedArchiveWriter.h"
 #include "KeyGenerator.h"
+#include "AES.h"
 
 IMPLEMENT_APPLICATION(UnrealPak, "UnrealPak");
 
@@ -20,13 +21,19 @@ struct FPakInputPair
 {
 	FString Source;
 	FString Dest;
+	bool bNeedsCompression;
+	bool bNeedEncryption;
 
 	FPakInputPair()
+		: bNeedsCompression(false)
+		, bNeedEncryption(false)
 	{}
 
 	FPakInputPair(const FString& InSource, const FString& InDest)
 		: Source(InSource)
 		, Dest(InDest)
+		, bNeedsCompression(false)
+		, bNeedEncryption(false)
 	{}
 
 	FORCEINLINE bool operator==(const FPakInputPair& Other) const
@@ -101,33 +108,135 @@ FString GetCommonRootPath(TArray<FPakInputPair>& FilesToAdd)
 }
 
 
-bool CopyFileToPak(FArchive& InPak, const FString& InMountPoint, const FPakInputPair& InFile, uint8*& InOutPersistentBuffer, int64& InOutBufferSize, FPakEntryPair& OutNewEntry)
+bool CopyFileToPak(FArchive& InPak, const FString& InMountPoint, const FPakInputPair& InFile, uint8*& InOutPersistentBuffer, int64& InOutBufferSize, uint8*& InPersistentCompressionBuffer, int64& InBufferCompressionSize, FPakEntryPair& OutNewEntry, ECompressionFlags CompressionMethod, const int32 CompressionBlockSize)
 {	
 	TAutoPtr<FArchive> FileHandle(IFileManager::Get().CreateFileReader(*InFile.Source));
 	bool bFileExists = FileHandle.IsValid();
 	if (bFileExists)
 	{
 		const int64 FileSize = FileHandle->TotalSize();
+		const int64 PaddedEncryptedFileSize = Align(FileSize, FAES::AESBlockSize); 
 		OutNewEntry.Filename = InFile.Dest.Mid(InMountPoint.Len());
 		OutNewEntry.Info.Offset = 0; // Don't serialize offsets here.
 		OutNewEntry.Info.Size = FileSize;
 		OutNewEntry.Info.UncompressedSize = FileSize;
-		OutNewEntry.Info.CompressionMethod = 0;
+		OutNewEntry.Info.CompressionMethod = CompressionMethod;
+		OutNewEntry.Info.bEncrypted = InFile.bNeedEncryption;
 
-		if (InOutBufferSize < FileSize)
+		if (InOutBufferSize < PaddedEncryptedFileSize)
 		{
-			InOutPersistentBuffer = (uint8*)FMemory::Realloc(InOutPersistentBuffer, FileSize);
+			InOutPersistentBuffer = (uint8*)FMemory::Realloc(InOutPersistentBuffer, PaddedEncryptedFileSize);
 			InOutBufferSize = FileSize;
 		}
 
+		// Don't store empty files with compression, it's wasteful
+		if (FileSize == 0) 
+		{
+			CompressionMethod = COMPRESS_None;
+		}
+
 		// Load to buffer
 		FileHandle->Serialize(InOutPersistentBuffer, FileSize);
-		// Calculate the buffer hash value
-		FSHA1::HashBuffer(InOutPersistentBuffer, FileSize, OutNewEntry.Info.Hash);
 
-		// Write to file
-		OutNewEntry.Info.Serialize(InPak, FPakInfo::PakFile_Version_Latest);
-		InPak.Serialize(InOutPersistentBuffer, FileSize);
+		if (CompressionMethod == COMPRESS_None)
+		{
+			int64 SizeToWrite = FileSize;
+			if (InFile.bNeedEncryption)
+			{
+				for(int64 FillIndex=FileSize; FillIndex < PaddedEncryptedFileSize && InFile.bNeedEncryption; ++FillIndex)
+				{
+					// Fill the trailing buffer with random bytes from file
+					InOutPersistentBuffer[FillIndex] = InOutPersistentBuffer[rand()%FileSize];
+				}
+
+				//Encrypt the buffer before writing it to disk
+				FAES::EncryptData(InOutPersistentBuffer, PaddedEncryptedFileSize);
+				// Update the size to be written
+				SizeToWrite = PaddedEncryptedFileSize;
+				OutNewEntry.Info.bEncrypted = true;
+			}
+
+			// Calculate the buffer hash value
+			FSHA1::HashBuffer(InOutPersistentBuffer,FileSize,OutNewEntry.Info.Hash);
+
+			// Write to file
+			OutNewEntry.Info.Serialize(InPak,FPakInfo::PakFile_Version_Latest);
+			InPak.Serialize(InOutPersistentBuffer,SizeToWrite);
+		}
+		else
+		{
+			// Build buffers for working
+			int64 UncompressedSize = FileSize;
+			int32 CompressionBufferSize = Align(FCompression::CompressMemoryBound(CompressionMethod, CompressionBlockSize), FAES::AESBlockSize);
+			if (CompressionBufferSize > InBufferCompressionSize)
+			{
+				InPersistentCompressionBuffer = (uint8*)FMemory::Realloc(InPersistentCompressionBuffer,CompressionBufferSize);
+				InBufferCompressionSize = CompressionBufferSize;
+			}
+			
+			int64 HeaderTell = InPak.Tell();
+			OutNewEntry.Info.CompressionBlocks.AddUninitialized((FileSize+CompressionBlockSize-1)/CompressionBlockSize);
+			OutNewEntry.Info.Serialize(InPak,FPakInfo::PakFile_Version_Latest);
+
+			int64 UncompressedBytes = 0;
+			int64 CompressedBytes = 0;
+			int32 CurrentBlock = 0;
+			FSHA1 Hash;
+			Hash.Reset();
+			while (UncompressedSize)
+			{
+				int32 BlockSize = (int32)FMath::Min<int64>(UncompressedSize, CompressionBlockSize);
+				int32 CompressedBlockSize = CompressionBufferSize;
+				OutNewEntry.Info.CompressionBlockSize = FMath::Max<uint32>(BlockSize, OutNewEntry.Info.CompressionBlockSize);
+				if (!FCompression::CompressMemory(CompressionMethod, InPersistentCompressionBuffer, CompressedBlockSize, InOutPersistentBuffer+UncompressedBytes, BlockSize))
+				{
+					return false;
+				}
+				UncompressedSize -= BlockSize;
+				UncompressedBytes += BlockSize;
+				CompressedBytes += CompressedBlockSize;
+
+				OutNewEntry.Info.CompressionBlocks[CurrentBlock].CompressedStart = InPak.Tell();
+				OutNewEntry.Info.CompressionBlocks[CurrentBlock].CompressedEnd = InPak.Tell()+CompressedBlockSize;
+				++CurrentBlock;
+
+				if (InFile.bNeedEncryption)
+				{
+					int32 EncryptionBlockPadding = Align(CompressedBlockSize,FAES::AESBlockSize);
+					for(int64 FillIndex=CompressedBlockSize; FillIndex < EncryptionBlockPadding; ++FillIndex)
+					{
+						// Fill the trailing buffer with random bytes from file
+						InPersistentCompressionBuffer[FillIndex] = InPersistentCompressionBuffer[rand()%CompressedBlockSize];
+					}
+					FAES::EncryptData(InPersistentCompressionBuffer, EncryptionBlockPadding);
+					CompressedBytes += EncryptionBlockPadding - CompressedBlockSize;
+
+					//Write out buffer
+					InPak.Serialize(InPersistentCompressionBuffer,EncryptionBlockPadding);
+					// Calculate the buffer hash value
+					Hash.Update(InPersistentCompressionBuffer,EncryptionBlockPadding);
+				}
+				else
+				{
+					//Write out buffer
+					InPak.Serialize(InPersistentCompressionBuffer, CompressedBlockSize);
+					// Calculate the buffer hash value
+					Hash.Update(InPersistentCompressionBuffer, CompressedBlockSize);
+				}
+			}
+
+			Hash.Final();
+
+			// Update file size & Hash
+			OutNewEntry.Info.Size = CompressedBytes;
+			Hash.GetHash(OutNewEntry.Info.Hash);
+
+			// Write out the corrected header
+			int64 CurrentTell = InPak.Tell();
+			InPak.Seek(HeaderTell);
+			OutNewEntry.Info.Serialize(InPak,FPakInfo::PakFile_Version_Latest);
+			InPak.Seek(CurrentTell);
+		}
 	}
 	return bFileExists;
 }
@@ -193,20 +302,8 @@ void ProcessCommandLine(int32 ArgC, TCHAR* ArgV[], TArray<FPakInputPair>& Entrie
 			for (int32 EntryIndex = 0; EntryIndex < Lines.Num(); EntryIndex++)
 			{
 				TArray<FString> SourceAndDest;
-				int32 Index = Lines[EntryIndex].Find("\" \"");
-				if (Index > 0 && Lines[EntryIndex].StartsWith("\"") && Lines[EntryIndex].EndsWith("\""))
-				{
-					SourceAndDest.Add(Lines[EntryIndex].Mid(1, Index - 1));
-					SourceAndDest.Add(Lines[EntryIndex].Mid(Index + 3, Lines[EntryIndex].Len() - Index - 4));
-					if (SourceAndDest[1].Len() < 1)
-					{
-						SourceAndDest.RemoveAt(1);
-					}
-				}
-				else
-				{
-					Lines[EntryIndex].ParseIntoArrayWS(&SourceAndDest);
-				}
+				TArray<FString> Switches;
+				FCommandLine::Parse(*Lines[EntryIndex], SourceAndDest, Switches);
 				FPakInputPair Input;
 				Input.Source = SourceAndDest[0];
 				FPaths::NormalizeFilename(Input.Source);
@@ -220,6 +317,20 @@ void ProcessCommandLine(int32 ArgC, TCHAR* ArgV[], TArray<FPakInputPair>& Entrie
 				}
 				FPaths::NormalizeFilename(Input.Dest);
 				FPakFile::MakeDirectoryFromPath(Input.Dest);
+
+				//check for compression switches
+				for (int32 Index = 0; Index < Switches.Num(); ++Index)
+				{
+					if (Switches[Index] == TEXT("compress"))
+					{
+						Input.bNeedsCompression = true;
+					}
+					if(Switches[Index] == TEXT("encrypt"))
+					{
+						Input.bNeedEncryption = true;
+					}
+				}
+
 				UE_LOG(LogPakFile, Display, TEXT("Added file Source: %s Dest: %s"), *Input.Source, *Input.Dest);
 				Entries.Add(Input);
 			}			
@@ -276,9 +387,13 @@ void CollectFilesToAdd(TArray<FPakInputPair>& OutFilesToAdd, const TArray<FPakIn
 	for (int32 Index = 0; Index < InEntries.Num(); Index++)
 	{
 		const FPakInputPair& Input = InEntries[Index];
+		const FString& Source = Input.Source;
+		bool bCompression = Input.bNeedsCompression;
+		bool bEncryption = Input.bNeedEncryption;
+
 
-		FString Filename = FPaths::GetCleanFilename(Input.Source);
-		FString Directory = FPaths::GetPath(Input.Source);
+		FString Filename = FPaths::GetCleanFilename(Source);
+		FString Directory = FPaths::GetPath(Source);
 		FPaths::MakeStandardFilename(Directory);
 		FPakFile::MakeDirectoryFromPath(Directory);
 
@@ -298,11 +413,20 @@ void CollectFilesToAdd(TArray<FPakInputPair>& OutFilesToAdd, const TArray<FPakIn
 				FileInput.Source = FoundFiles[FileIndex];
 				FPaths::MakeStandardFilename(FileInput.Source);
 				FileInput.Dest = FileInput.Source.Replace(*Directory, *Input.Dest, ESearchCase::IgnoreCase);
+				FileInput.bNeedsCompression = bCompression;
+				FileInput.bNeedEncryption = bEncryption;
 				if (!AddedFiles.Contains(FileInput.Source))
 				{
 					OutFilesToAdd.Add(FileInput);
 					AddedFiles.Add(FileInput.Source);
 				}
+				else
+				{
+					int32 FoundIndex;
+					OutFilesToAdd.Find(FileInput,FoundIndex);
+					OutFilesToAdd[FoundIndex].bNeedEncryption |= bEncryption;
+					OutFilesToAdd[FoundIndex].bNeedsCompression |= bCompression;
+				}
 			}
 		}
 		else
@@ -312,9 +436,23 @@ void CollectFilesToAdd(TArray<FPakInputPair>& OutFilesToAdd, const TArray<FPakIn
 			FileInput.Source = Input.Source;
 			FPaths::MakeStandardFilename(FileInput.Source);
 			FileInput.Dest = FileInput.Source.Replace(*Directory, *Input.Dest, ESearchCase::IgnoreCase);
-			OutFilesToAdd.AddUnique(FileInput);
+			FileInput.bNeedEncryption = bEncryption;
+			FileInput.bNeedsCompression = bCompression;
+
+			if (AddedFiles.Contains(FileInput.Source))
+			{
+				int32 FoundIndex;
+				OutFilesToAdd.Find(FileInput, FoundIndex);
+				OutFilesToAdd[FoundIndex].bNeedEncryption |= bEncryption;
+				OutFilesToAdd[FoundIndex].bNeedsCompression |= bCompression;
+			}
+			else
+			{
+				OutFilesToAdd.AddUnique(FileInput);
+			}
 		}
 	}
+
 	// Sort alphabetically
 	struct FInputPairSort
 	{
@@ -338,19 +476,69 @@ void CollectFilesToAdd(TArray<FPakInputPair>& OutFilesToAdd, const TArray<FPakIn
 	UE_LOG(LogPakFile, Display, TEXT("Collected %d files in %.2lfs."), OutFilesToAdd.Num(), FPlatformTime::Seconds() - StartTime);
 }
 
-bool BufferedCopyFile(FArchive& Dest, FArchive& Source, const int64 FileSize, void* Buffer, const int64 BufferSize)
+bool BufferedCopyFile(FArchive& Dest, FArchive& Source, const FPakEntry& Entry, void* Buffer, int64 BufferSize)
 {	
-	int64 RemainingSizeToCopy = FileSize;
+	// Align down
+	BufferSize = BufferSize & ~(FAES::AESBlockSize-1);
+	int64 RemainingSizeToCopy = Entry.Size;
 	while (RemainingSizeToCopy > 0)
 	{
 		const int64 SizeToCopy = FMath::Min(BufferSize, RemainingSizeToCopy);
-		Source.Serialize(Buffer, SizeToCopy);
+		// If file is encrypted so we need to account for padding
+		int64 SizeToRead = Entry.bEncrypted ? Align(SizeToCopy,FAES::AESBlockSize) : SizeToCopy;
+
+		Source.Serialize(Buffer,SizeToRead);
+		if (Entry.bEncrypted)
+		{
+			FAES::DecryptData((uint8*)Buffer,SizeToRead);
+		}
 		Dest.Serialize(Buffer, SizeToCopy);
-		RemainingSizeToCopy -= SizeToCopy;
+		RemainingSizeToCopy -= SizeToRead;
 	}
 	return true;
 }
 
+bool UncompressCopyFile(FArchive& Dest, FArchive& Source, const FPakEntry& Entry, uint8*& PersistentBuffer, int64& BufferSize)
+{
+	if (Entry.UncompressedSize == 0)
+	{
+		return false;
+	}
+
+	int64 WorkingSize = Entry.CompressionBlockSize;
+	int32 MaxCompressionBlockSize = FCompression::CompressMemoryBound((ECompressionFlags)Entry.CompressionMethod, WorkingSize);
+	WorkingSize += MaxCompressionBlockSize;
+	if (BufferSize < WorkingSize)
+	{
+		PersistentBuffer = (uint8*)FMemory::Realloc(PersistentBuffer, WorkingSize);
+		BufferSize = WorkingSize;
+	}
+
+	uint8* UncompressedBuffer = PersistentBuffer+MaxCompressionBlockSize;
+
+	for (uint32 BlockIndex=0, BlockIndexNum=Entry.CompressionBlocks.Num(); BlockIndex < BlockIndexNum; ++BlockIndex)
+	{
+		uint32 CompressedBlockSize = Entry.CompressionBlocks[BlockIndex].CompressedEnd - Entry.CompressionBlocks[BlockIndex].CompressedStart;
+		uint32 UncompressedBlockSize = (uint32)FMath::Min<int64>(Entry.UncompressedSize - Entry.CompressionBlockSize*BlockIndex, Entry.CompressionBlockSize);
+		Source.Seek(Entry.CompressionBlocks[BlockIndex].CompressedStart);
+		uint32 SizeToRead = Entry.bEncrypted ? Align(CompressedBlockSize, FAES::AESBlockSize) : CompressedBlockSize;
+		Source.Serialize(PersistentBuffer, SizeToRead);
+
+		if (Entry.bEncrypted)
+		{
+			FAES::DecryptData(PersistentBuffer, SizeToRead);
+		}
+
+		if(!FCompression::UncompressMemory((ECompressionFlags)Entry.CompressionMethod,UncompressedBuffer,UncompressedBlockSize,PersistentBuffer,CompressedBlockSize))
+		{
+			return false;
+		}
+		Dest.Serialize(UncompressedBuffer,UncompressedBlockSize);
+	}
+
+	return true;
+}
+
 /**
  * Creates a pak file writer. This can be a signed writer if the encryption keys are specified in the command line
  */
@@ -399,7 +587,7 @@ FArchive* CreatePakWriter(const TCHAR* Filename)
 	return Writer;
 }
 
-bool CreatePakFile(const TCHAR* Filename, TArray<FPakInputPair>& FilesToAdd)
+bool CreatePakFile(const TCHAR* Filename, TArray<FPakInputPair>& FilesToAdd, const int32 CompressionBlockSize)
 {	
 	const double StartTime = FPlatformTime::Seconds();
 
@@ -416,18 +604,32 @@ bool CreatePakFile(const TCHAR* Filename, TArray<FPakInputPair>& FilesToAdd)
 	FString MountPoint = GetCommonRootPath(FilesToAdd);
 	uint8* ReadBuffer = NULL;
 	int64 BufferSize = 0;
+	uint8* CompressionBuffer = NULL;
+	int64 CompressionBufferSize = 0;
+	ECompressionFlags CompressionMethod = COMPRESS_None;
 
 	for (int32 FileIndex = 0; FileIndex < FilesToAdd.Num(); FileIndex++)
 	{
 		//  Remember the offset but don't serialize it with the entry header.
 		const int64 NewEntryOffset = PakFileHandle->Tell();
 		FPakEntryPair NewEntry;
-		if (CopyFileToPak(*PakFileHandle, MountPoint, FilesToAdd[FileIndex], ReadBuffer, BufferSize, NewEntry) == true)
+
+		//check if this file requested to be compression
+		CompressionMethod = FilesToAdd[FileIndex].bNeedsCompression ? COMPRESS_Default : COMPRESS_None;
+
+		if (CopyFileToPak(*PakFileHandle, MountPoint, FilesToAdd[FileIndex], ReadBuffer, BufferSize, CompressionBuffer, CompressionBufferSize, NewEntry, CompressionMethod, CompressionBlockSize) == true)
 		{
 			// Update offset now and store it in the index (and only in index)
 			NewEntry.Info.Offset = NewEntryOffset;
 			Index.Add(NewEntry);
-			UE_LOG(LogPakFile, Display, TEXT("Added file \"%s\", %lld bytes."), *NewEntry.Filename, NewEntry.Info.Size);
+			if (FilesToAdd[FileIndex].bNeedsCompression)
+			{
+				UE_LOG(LogPakFile, Display, TEXT("Added compressed file \"%s\", Compressed Size %lld bytes, Original Size %lld bytes."), *NewEntry.Filename, NewEntry.Info.Size, NewEntry.Info.UncompressedSize);
+			}
+			else
+			{
+				UE_LOG(LogPakFile, Display, TEXT("Added file \"%s\", %lld bytes."), *NewEntry.Filename, NewEntry.Info.Size);
+			}
 		}
 		else
 		{
@@ -534,6 +736,8 @@ bool ExtractFilesFromPak(const TCHAR* InPakFilename, const TCHAR* InDestPath)
 		FArchive& PakReader = *PakFile.GetSharedReader(NULL);
 		const int64 BufferSize = 8 * 1024 * 1024; // 8MB buffer for extracting
 		void* Buffer = FMemory::Malloc(BufferSize);
+		int64 CompressionBufferSize = 0;
+		uint8* PersistantCompressionBuffer = NULL;
 		int32 ErrorCount = 0;
 		int32 FileCount = 0;
 
@@ -550,7 +754,14 @@ bool ExtractFilesFromPak(const TCHAR* InPakFilename, const TCHAR* InDestPath)
 				TAutoPtr<FArchive> FileHandle(IFileManager::Get().CreateFileWriter(*DestFilename));
 				if (FileHandle.IsValid())
 				{
-					BufferedCopyFile(*FileHandle, PakReader, Entry.Size, Buffer, BufferSize);
+					if (Entry.CompressionMethod == COMPRESS_None)
+					{
+						BufferedCopyFile(*FileHandle, PakReader, Entry, Buffer, BufferSize);
+					}
+					else
+					{
+						UncompressCopyFile(*FileHandle, PakReader, Entry, PersistantCompressionBuffer, CompressionBufferSize);
+					}
 					UE_LOG(LogPakFile, Display, TEXT("Extracted \"%s\" to \"%s\"."), *It.Filename(), *DestFilename);
 				}
 				else
@@ -566,6 +777,7 @@ bool ExtractFilesFromPak(const TCHAR* InPakFilename, const TCHAR* InDestPath)
 			}
 		}
 		FMemory::Free(Buffer);
+		FMemory::Free(PersistantCompressionBuffer);
 
 		UE_LOG(LogPakFile, Error, TEXT("Finished extracting %d files (including %d errors)."), FileCount, ErrorCount);
 
@@ -607,6 +819,7 @@ INT32_MAIN_INT32_ARGC_TCHAR_ARGV()
 		return 1;
 	}
 
+	int32 CompressionBlockSize = 64*1024;
 	int32 Result = 0;
 	FString KeyFilename;
 	if (FParse::Value(FCommandLine::Get(), TEXT("GenerateKeys="), KeyFilename, false))
@@ -660,7 +873,7 @@ INT32_MAIN_INT32_ARGC_TCHAR_ARGV()
 				TArray<FPakInputPair> FilesToAdd;
 				CollectFilesToAdd(FilesToAdd, Entries, OrderMap);
 
-				Result = CreatePakFile(*PakFilename, FilesToAdd) ? 0 : 1;
+				Result = CreatePakFile(*PakFilename, FilesToAdd, CompressionBlockSize) ? 0 : 1;
 			}
 		}
 	}
diff --git a/Engine/Source/Runtime/Core/Private/Misc/Compression.cpp b/Engine/Source/Runtime/Core/Private/Misc/Compression.cpp
--- a/Engine/Source/Runtime/Core/Private/Misc/Compression.cpp
+++ b/Engine/Source/Runtime/Core/Private/Misc/Compression.cpp
@@ -73,6 +73,59 @@ uint64 FCompression::CompressorSrcBytes	= 0;
 /** Nubmer of bytes after compression.		*/
 uint64 FCompression::CompressorDstBytes	= 0;
 
+static ECompressionFlags CheckGlobalCompressionFlags(ECompressionFlags Flags)
+{
+	static bool GAlwaysBiasCompressionForSize = false;
+	if(FPlatformProperties::HasEditorOnlyData())
+	{
+		static bool GTestedCmdLine = false;
+		if(!GTestedCmdLine && FCommandLine::IsInitialized())
+		{
+			GTestedCmdLine = true;
+			// Override compression settings wrt size.
+			GAlwaysBiasCompressionForSize = FParse::Param(FCommandLine::Get(),TEXT("BIASCOMPRESSIONFORSIZE"));
+		}
+	}
+
+	// Always bias for speed if option is set.
+	if(GAlwaysBiasCompressionForSize)
+	{
+		int32 NewFlags = Flags;
+		NewFlags &= ~COMPRESS_BiasSpeed;
+		NewFlags |= COMPRESS_BiasMemory;
+		Flags = (ECompressionFlags)NewFlags;
+	}
+
+	return Flags;
+}
+
+/**
+* Thread-safe abstract compression routine to query memory requirements for a compression operation.
+*
+* @param	Flags						Flags to control what method to use and optionally control memory vs speed
+* @param	UncompressedSize			Size of uncompressed data in bytes
+* @return The maximum possible bytes needed for compression of data buffer of size UncompressedSize
+*/
+int32 FCompression::CompressMemoryBound( ECompressionFlags Flags, int32 UncompressedSize ) 
+{
+	int32 CompressionBound = UncompressedSize;
+	// make sure a valid compression scheme was provided
+	check(Flags & COMPRESS_ZLIB);
+
+	Flags = CheckGlobalCompressionFlags(Flags);
+
+	switch(Flags & COMPRESSION_FLAGS_TYPE_MASK)
+	{
+	case COMPRESS_ZLIB:
+		CompressionBound = compressBound(UncompressedSize);
+		break;
+	default:
+		break;
+	}
+
+	return CompressionBound;
+}
+
 /**
  * Thread-safe abstract compression routine. Compresses memory from uncompressed buffer and writes it to compressed
  * buffer. Updates CompressedSize with size of compressed data. Compression controlled by the passed in flags.
@@ -93,26 +146,7 @@ bool FCompression::CompressMemory( ECompressionFlags Flags, void* CompressedBuff
 
 	bool bCompressSucceeded = false;
 
-	static bool GAlwaysBiasCompressionForSize = false;
-	if (FPlatformProperties::HasEditorOnlyData())
-	{
-		static bool GTestedCmdLine = false;
-		if (!GTestedCmdLine && FCommandLine::IsInitialized())
-		{
-			GTestedCmdLine = true;
-			// Override compression settings wrt size.
-			GAlwaysBiasCompressionForSize = FParse::Param( FCommandLine::Get(), TEXT("BIASCOMPRESSIONFORSIZE") );
-		}
-	}
-
-	// Always bias for speed if option is set.
-	if( GAlwaysBiasCompressionForSize )
-	{
-		int32 NewFlags = Flags;
-		NewFlags &= ~COMPRESS_BiasSpeed;
-		NewFlags |= COMPRESS_BiasMemory;
-		Flags = (ECompressionFlags) NewFlags;
-	}
+	Flags = CheckGlobalCompressionFlags(Flags);
 
 	switch(Flags & COMPRESSION_FLAGS_TYPE_MASK)
 	{
@@ -168,7 +202,7 @@ bool FCompression::UncompressMemory( ECompressionFlags Flags, void* Uncompressed
 			UE_LOG(LogCompression, Warning, TEXT("FCompression::UncompressMemory - This compression type not supported"));
 			bUncompressSucceeded = false;
 	}
-	INC_FLOAT_STAT_BY(STAT_UncompressorTime,(float)(FPlatformTime::Seconds()-UncompressorStartTime));
+	STAT(if (FThreadStats::IsThreadingReady()) { INC_FLOAT_STAT_BY(STAT_UncompressorTime,(float)(FPlatformTime::Seconds()-UncompressorStartTime))} );
 	
 	return bUncompressSucceeded;
 }
diff --git a/Engine/Source/Runtime/Core/Public/Misc/Compression.h b/Engine/Source/Runtime/Core/Public/Misc/Compression.h
--- a/Engine/Source/Runtime/Core/Public/Misc/Compression.h
+++ b/Engine/Source/Runtime/Core/Public/Misc/Compression.h
@@ -51,6 +51,15 @@ struct FCompression
 	/** Number of bytes after compression.		*/
 	CORE_API static uint64 CompressorDstBytes;
 
+	/**
+	 * Thread-safe abstract compression routine to query memory requirements for a compression operation.
+	 *
+	 * @param	Flags						Flags to control what method to use and optionally control memory vs speed
+	 * @param	UncompressedSize			Size of uncompressed data in bytes
+	 * @return The maximum possible bytes needed for compression of data buffer of size UncompressedSize
+	 */
+	CORE_API static int32 CompressMemoryBound( ECompressionFlags Flags, int32 UncompressedSize );
+
 	/**
 	 * Thread-safe abstract compression routine. Compresses memory from uncompressed buffer and writes it to compressed
 	 * buffer. Updates CompressedSize with size of compressed data. Compression controlled by the passed in flags.
diff --git a/Engine/Source/Runtime/PakFile/Private/IPlatformFilePak.cpp b/Engine/Source/Runtime/PakFile/Private/IPlatformFilePak.cpp
--- a/Engine/Source/Runtime/PakFile/Private/IPlatformFilePak.cpp
+++ b/Engine/Source/Runtime/PakFile/Private/IPlatformFilePak.cpp
@@ -10,9 +10,199 @@
 #include "BigInt.h"
 #include "SignedArchiveReader.h"
 #include "PublicKey.inl"
+#include "AES.h"
 
 DEFINE_LOG_CATEGORY(LogPakFile);
 
+
+/**
+ * Class to handle correctly reading from a compressed file within a compressed package
+ */
+class FPakSimpleEncryption
+{
+public:
+	enum
+	{
+		Alignment = FAES::AESBlockSize,
+	};
+
+	static FORCEINLINE int64 AlignReadRequest(int64 Size) 
+	{
+		return Align(Size, Alignment);
+	}
+
+	static FORCEINLINE void DecryptBlock(void* Data, int64 Size)
+	{
+#ifdef AES_KEY
+		FAES::DecryptData((uint8*)Data, Size);
+#endif
+	}
+};
+
+/**
+ * Thread local class to manage working buffers for file compression
+ */
+class FCompressionScratchBuffers : public TThreadSingleton<FCompressionScratchBuffers>
+{
+public:
+	FCompressionScratchBuffers()
+		: TempBufferSize(0)
+		, ScratchBufferSize(0)
+	{}
+
+	int64				TempBufferSize;
+	TAutoPtr<uint8>		TempBuffer;
+	int64				ScratchBufferSize;
+	TAutoPtr<uint8>		ScratchBuffer;
+
+	void EnsureBufferSpace(int64 CompressionBlockSize, int64 ScrachSize)
+	{
+		if(TempBufferSize < CompressionBlockSize)
+		{
+			TempBufferSize = CompressionBlockSize;
+			TempBuffer.Reset((uint8*)FMemory::Malloc(TempBufferSize));
+		}
+		if(ScratchBufferSize < ScrachSize)
+		{
+			ScratchBufferSize = ScrachSize;
+			ScratchBuffer.Reset((uint8*)FMemory::Malloc(ScratchBufferSize));
+		}
+	}
+};
+
+/**
+ * Class to handle correctly reading from a compressed file within a pak
+ */
+template< typename EncryptionPolicy = FPakNoEncryption >
+class FPakCompressedReaderPolicy
+{
+public:
+	class FPakUncompressTask : public FNonAbandonableTask
+	{
+	public:
+		uint8*				UncompressedBuffer;
+		int32				UncompressedSize;
+		uint8*				CompressedBuffer;
+		int32				CompressedSize;
+		ECompressionFlags	Flags;
+		void*				CopyOut;
+		int64				CopyOffset;
+		int64				CopyLength;
+
+		void DoWork()
+		{
+			// Decrypt and Uncompress from memory to memory.
+			int64 EncryptionSize = EncryptionPolicy::AlignReadRequest(CompressedSize);
+			EncryptionPolicy::DecryptBlock(CompressedBuffer, EncryptionSize);
+			FCompression::UncompressMemory(Flags, UncompressedBuffer, UncompressedSize, CompressedBuffer, CompressedSize, false);
+			if (CopyOut)
+			{
+				FMemory::Memcpy(CopyOut, UncompressedBuffer+CopyOffset, CopyLength);
+			}
+		}
+		static const TCHAR *Name()
+		{
+			return TEXT("FPakUncompressTask");
+		}
+	};
+
+	FPakCompressedReaderPolicy(const FPakFile& InPakFile, const FPakEntry& InPakEntry, FArchive* InPakReader)
+		: PakFile(InPakFile)
+		, PakEntry(InPakEntry)
+		, PakReader(InPakReader)
+	{
+	}
+
+	/** Pak file that own this file data */
+	const FPakFile&		PakFile;
+	/** Pak file entry for this file. */
+	const FPakEntry&	PakEntry;
+	/** Pak file archive to read the data from. */
+	FArchive*			PakReader;
+
+	FORCEINLINE int64 FileSize() const
+	{
+		return PakEntry.UncompressedSize;
+	}
+
+	void Serialize(int64 DesiredPosition, void* V, int64 Length)
+	{
+		const int32 CompressionBlockSize = PakEntry.CompressionBlockSize;
+		uint32 CompressionBlockIndex = DesiredPosition / CompressionBlockSize;
+		uint8* WorkingBuffers[2];
+		int64 DirectCopyStart = DesiredPosition % PakEntry.CompressionBlockSize;
+		FAsyncTask<FPakUncompressTask> UncompressTask;
+		FCompressionScratchBuffers& ScratchSpace = FCompressionScratchBuffers::Get();
+		bool bStartedUncompress = false;
+
+		int64 WorkingBufferRequiredSize = FCompression::CompressMemoryBound((ECompressionFlags)PakEntry.CompressionMethod,CompressionBlockSize);
+		WorkingBufferRequiredSize = EncryptionPolicy::AlignReadRequest(WorkingBufferRequiredSize);
+		ScratchSpace.EnsureBufferSpace(CompressionBlockSize, WorkingBufferRequiredSize*2);
+		WorkingBuffers[0] = ScratchSpace.ScratchBuffer;
+		WorkingBuffers[1] = ScratchSpace.ScratchBuffer + WorkingBufferRequiredSize;
+
+		while (Length > 0)
+		{
+			const FPakCompressedBlock& Block = PakEntry.CompressionBlocks[CompressionBlockIndex];
+			int64 Pos = CompressionBlockIndex * CompressionBlockSize;
+			int64 CompressedBlockSize = Block.CompressedEnd-Block.CompressedStart;
+			int64 UncompressedBlockSize = FMath::Min<int64>(PakEntry.UncompressedSize-Pos, PakEntry.CompressionBlockSize);
+			int64 ReadSize = EncryptionPolicy::AlignReadRequest(CompressedBlockSize);
+			int64 WriteSize = FMath::Min<int64>(UncompressedBlockSize - DirectCopyStart, Length);
+			PakReader->Seek(Block.CompressedStart);
+			PakReader->Serialize(WorkingBuffers[CompressionBlockIndex & 1],ReadSize);
+			if (bStartedUncompress)
+			{
+				UncompressTask.EnsureCompletion();
+				bStartedUncompress = false;
+			}
+
+			FPakUncompressTask& TaskDetails = UncompressTask.GetTask();
+			if (DirectCopyStart == 0 && Length >= CompressionBlockSize)
+			{
+				// Block can be decompressed directly into output buffer
+				TaskDetails.Flags = (ECompressionFlags)PakEntry.CompressionMethod;
+				TaskDetails.UncompressedBuffer = (uint8*)V;
+				TaskDetails.UncompressedSize = UncompressedBlockSize;
+				TaskDetails.CompressedBuffer = WorkingBuffers[CompressionBlockIndex & 1];
+				TaskDetails.CompressedSize = CompressedBlockSize;
+				TaskDetails.CopyOut = nullptr;
+			}
+			else
+			{
+				// Block needs to be copied from a working buffer
+				TaskDetails.Flags = (ECompressionFlags)PakEntry.CompressionMethod;
+				TaskDetails.UncompressedBuffer = (uint8*)ScratchSpace.TempBuffer;
+				TaskDetails.UncompressedSize = UncompressedBlockSize;
+				TaskDetails.CompressedBuffer = WorkingBuffers[CompressionBlockIndex & 1];
+				TaskDetails.CompressedSize = CompressedBlockSize;
+				TaskDetails.CopyOut = V;
+				TaskDetails.CopyOffset = DirectCopyStart;
+				TaskDetails.CopyLength = WriteSize;
+			}
+			
+			if (Length == WriteSize)
+			{
+				UncompressTask.StartSynchronousTask();
+			}
+			else
+			{
+				UncompressTask.StartBackgroundTask();
+			}
+			bStartedUncompress = true;
+			V = (void*)((uint8*)V + WriteSize);
+			Length -= WriteSize;
+			DirectCopyStart = 0;
+			++CompressionBlockIndex;
+		}
+
+		if(bStartedUncompress)
+		{
+			UncompressTask.EnsureCompletion();
+		}
+	}
+};
+
 bool FPakEntry::VerifyPakEntriesMatch(const FPakEntry& FileEntryA, const FPakEntry& FileEntryB)
 {
 	bool bResult = true;
@@ -491,7 +681,25 @@ IFileHandle* FPakPlatformFile::CreatePakFileHandle(const TCHAR* Filename, FPakFi
 	FArchive* PakReader = PakFile->GetSharedReader(LowerLevel);
 
 	// Create the handle.
-	Result = new FPakFileHandle(*PakFile, *FileEntry, PakReader, true);		
+	if (FileEntry->CompressionMethod != COMPRESS_None && PakFile->GetInfo().Version >= FPakInfo::PakFile_Version_CompressionEncryption)
+	{
+		if (FileEntry->bEncrypted)
+		{
+			Result = new FPakFileHandle< FPakCompressedReaderPolicy<FPakSimpleEncryption> >(*PakFile, *FileEntry, PakReader, true);
+		}
+		else
+		{
+			Result = new FPakFileHandle< FPakCompressedReaderPolicy<> >(*PakFile, *FileEntry, PakReader, true);
+		}
+	}
+	else if (FileEntry->bEncrypted)
+	{
+		Result = new FPakFileHandle< FPakReaderPolicy<FPakSimpleEncryption> >(*PakFile, *FileEntry, PakReader, true);
+	}
+	else
+	{
+		Result = new FPakFileHandle<>(*PakFile, *FileEntry, PakReader, true);
+	}
 
 	return Result;
 }
diff --git a/Engine/Source/Runtime/PakFile/Public/IPlatformFilePak.h b/Engine/Source/Runtime/PakFile/Public/IPlatformFilePak.h
--- a/Engine/Source/Runtime/PakFile/Public/IPlatformFilePak.h
+++ b/Engine/Source/Runtime/PakFile/Public/IPlatformFilePak.h
@@ -22,8 +22,9 @@ struct FPakInfo
 	{
 		PakFile_Version_Initial = 1,
 		PakFile_Version_NoTimestamps = 2,
+		PakFile_Version_CompressionEncryption = 3,
 
-		PakFile_Version_Latest = PakFile_Version_NoTimestamps
+		PakFile_Version_Latest = PakFile_Version_CompressionEncryption
 	};
 
 	/** Pak file magic value. */
@@ -80,6 +81,34 @@ struct FPakInfo
 	}
 };
 
+/**
+ * Struct storing offsets and sizes of a compressed block
+ */
+struct FPakCompressedBlock
+{
+	/** Offset of the start of a compression block. Offset is absolute */
+	int64 CompressedStart;
+	/** Offset of the end of a compression block. This may not align completely with the start of the next block. Offset is absolute */
+	int64 CompressedEnd;
+
+	bool operator == (const FPakCompressedBlock& B) const
+	{
+		return CompressedStart == B.CompressedStart && CompressedEnd == B.CompressedEnd;
+	}
+
+	bool operator != (const FPakCompressedBlock& B) const
+	{
+		return !(*this == B);
+	}
+};
+
+FORCEINLINE FArchive& operator<<(FArchive& Ar, FPakCompressedBlock& Block)
+{
+	Ar << Block.CompressedStart;
+	Ar << Block.CompressedEnd;
+	return Ar;
+}
+
 /**
  * Struct holding info about a single file stored in pak file.
  */
@@ -95,6 +124,12 @@ struct FPakEntry
 	int32 CompressionMethod;
 	/** File SHA1 value. */
 	uint8 Hash[20];
+	/** Array of compression blocks that describe how to decompress this pak entry */
+	TArray<FPakCompressedBlock> CompressionBlocks;
+	/** Size of a compressed block in the file */
+	uint32 CompressionBlockSize;
+	/** True is file is encrypted */
+	uint8 bEncrypted;
 	/** Flag is set to true when FileHeader has been checked against PakHeader. It is not serialized */
 	mutable bool  Verified;
 
@@ -106,6 +141,8 @@ struct FPakEntry
 		, Size(0)
 		, UncompressedSize(0)
 		, CompressionMethod(0)
+		, CompressionBlockSize(0)
+		, bEncrypted(false)
 		, Verified(false)
 	{
 		FMemory::Memset(Hash, 0, sizeof(Hash));
@@ -119,6 +156,14 @@ struct FPakEntry
 	int64 GetSerializedSize(int32 Version) const
 	{
 		int64 SerializedSize = sizeof(Offset) + sizeof(Size) + sizeof(UncompressedSize) + sizeof(CompressionMethod) + sizeof(Hash);
+		if (Version >= FPakInfo::PakFile_Version_CompressionEncryption)
+		{
+			SerializedSize += sizeof(bEncrypted) + sizeof(CompressionBlockSize);
+			if(CompressionMethod != COMPRESS_None)
+			{
+				SerializedSize += sizeof(FPakCompressedBlock) * CompressionBlocks.Num() + sizeof(int32);
+			}
+		}
 		if (Version < FPakInfo::PakFile_Version_NoTimestamps)
 		{
 			// Timestamp
@@ -137,7 +182,10 @@ struct FPakEntry
 		return Size == B.Size && 
 			UncompressedSize == B.UncompressedSize &&
 			CompressionMethod == B.CompressionMethod &&
-			FMemory::Memcmp(Hash, B.Hash, sizeof(Hash)) == 0;
+			bEncrypted == B.bEncrypted &&
+			CompressionBlockSize == B.CompressionBlockSize &&
+			FMemory::Memcmp(Hash, B.Hash, sizeof(Hash)) == 0 &&
+			CompressionBlocks == B.CompressionBlocks;
 	}
 
 	/**
@@ -150,7 +198,10 @@ struct FPakEntry
 		return Size != B.Size || 
 			UncompressedSize != B.UncompressedSize ||
 			CompressionMethod != B.CompressionMethod ||
-			FMemory::Memcmp(Hash, B.Hash, sizeof(Hash)) != 0;
+			bEncrypted != B.bEncrypted ||
+			CompressionBlockSize != B.CompressionBlockSize ||
+			FMemory::Memcmp(Hash, B.Hash, sizeof(Hash)) != 0 ||
+			CompressionBlocks != B.CompressionBlocks;
 	}
 
 	/**
@@ -171,6 +222,15 @@ struct FPakEntry
 			Ar << Timestamp;
 		}
 		Ar.Serialize(Hash, sizeof(Hash));
+		if (Version >= FPakInfo::PakFile_Version_CompressionEncryption)
+		{
+			if(CompressionMethod != COMPRESS_None)
+			{
+				Ar << CompressionBlocks;
+			}
+			Ar << bEncrypted;
+			Ar << CompressionBlockSize;
+		}
 	}
 
 	/**
@@ -524,22 +584,102 @@ public:
 };
 
 /**
- * File handle to read from pak file.
+ * Placeholder Class
  */
-class PAKFILE_API FPakFileHandle : public IFileHandle
-{	
+class PAKFILE_API FPakNoEncryption
+{
+public:
+	enum 
+	{
+		Alignment = 1,
+	};
+
+	static FORCEINLINE int64 AlignReadRequest(int64 Size) 
+	{
+		return Size;
+	}
+
+	static FORCEINLINE void DecryptBlock(void* Data, int64 Size)
+	{
+		// Nothing needs to be done here
+	}
+};
+
+template< typename EncryptionPolicy = FPakNoEncryption >
+class PAKFILE_API FPakReaderPolicy
+{
+public:
 	/** Pak file that own this file data */
-	const FPakFile& PakFile;
+	const FPakFile&		PakFile;
 	/** Pak file entry for this file. */
-	const FPakEntry& PakEntry;
+	const FPakEntry&	PakEntry;
 	/** Pak file archive to read the data from. */
-	FArchive* PakReader;
+	FArchive*			PakReader;
+	/** Offset to the file in pak (including the file header). */
+	int64				OffsetToFile;
+
+	FPakReaderPolicy(const FPakFile& InPakFile,const FPakEntry& InPakEntry,FArchive* InPakReader)
+		: PakFile(InPakFile)
+		, PakEntry(InPakEntry)
+		, PakReader(InPakReader)
+	{
+		OffsetToFile = PakEntry.Offset + PakEntry.GetSerializedSize(PakFile.GetInfo().Version);
+	}
+
+	FORCEINLINE int64 FileSize() const 
+	{
+		return PakEntry.Size;
+	}
+
+	void Serialize(int64 DesiredPosition, void* V, int64 Length)
+	{
+		uint8 TempBuffer[EncryptionPolicy::Alignment];
+		if (EncryptionPolicy::AlignReadRequest(DesiredPosition) != DesiredPosition)
+		{
+			int64 Start = DesiredPosition & ~(EncryptionPolicy::Alignment-1);
+			int64 Offset = DesiredPosition - Start;
+			int32 CopySize = EncryptionPolicy::Alignment-(DesiredPosition-Start);
+			PakReader->Seek(OffsetToFile + Start);
+			PakReader->Serialize(TempBuffer, EncryptionPolicy::Alignment);
+			EncryptionPolicy::DecryptBlock(TempBuffer, EncryptionPolicy::Alignment);
+			FMemory::Memcpy(V, TempBuffer+Offset, CopySize);
+			V = (void*)((uint8*)V + CopySize);
+			DesiredPosition += CopySize;
+			Length -= CopySize;
+			check(DesiredPosition % EncryptionPolicy::Alignment == 0);
+		}
+		else
+		{
+			PakReader->Seek(OffsetToFile + DesiredPosition);
+		}
+		
+		int64 CopySize = Length & ~(EncryptionPolicy::Alignment-1);
+		PakReader->Serialize(V, CopySize);
+		EncryptionPolicy::DecryptBlock(V, CopySize);
+		Length -= CopySize;
+		V = (void*)((uint8*)V + CopySize);
+
+		if (Length > 0)
+		{
+			PakReader->Serialize(TempBuffer, EncryptionPolicy::Alignment);
+			EncryptionPolicy::DecryptBlock(TempBuffer, EncryptionPolicy::Alignment);
+			FMemory::Memcpy(V, TempBuffer, Length);
+		}
+	}
+};
+
+/**
+ * File handle to read from pak file.
+ */
+template< typename ReaderPolicy = FPakReaderPolicy<> >
+class PAKFILE_API FPakFileHandle : public IFileHandle
+{	
 	/** True if PakReader is shared and should not be deleted by this handle. */
 	const bool bSharedReader;
-	/** Offset to the file in pak (including the file header). */
-	int64 OffsetToFile;
 	/** Current read position. */
 	int64 ReadPos;
+	/** Class that controls reading from pak file */
+	ReaderPolicy Reader;
 
 public:
 
@@ -551,13 +691,10 @@ public:
 	 * @param InPakFile Pak file.
 	 */
 	FPakFileHandle(const FPakFile& InPakFile, const FPakEntry& InPakEntry, FArchive* InPakReader, bool bIsSharedReader)
-		: PakFile(InPakFile)
-		, PakEntry(InPakEntry)
-		, PakReader(InPakReader)
-		, bSharedReader(bIsSharedReader)
+		: bSharedReader(bIsSharedReader)
 		, ReadPos(0)
+		, Reader(InPakFile, InPakEntry, InPakReader)
 	{
-		OffsetToFile = PakEntry.Offset + PakEntry.GetSerializedSize(PakFile.GetInfo().Version);
 	}
 
 	/**
@@ -567,7 +704,7 @@ public:
 	{
 		if (!bSharedReader)
 		{
-			delete PakReader;
+			delete Reader.PakReader;
 		}
 	}
 
@@ -578,7 +715,7 @@ public:
 	}
 	virtual bool Seek(int64 NewPosition) override
 	{
-		if (NewPosition > PakEntry.Size || NewPosition < 0)
+		if (NewPosition > Reader.FileSize() || NewPosition < 0)
 		{
 			return false;
 		}
@@ -587,19 +724,19 @@ public:
 	}
 	virtual bool SeekFromEnd(int64 NewPositionRelativeToEnd) override
 	{
-		return Seek(PakEntry.Size - NewPositionRelativeToEnd);
+		return Seek(Reader.FileSize() - NewPositionRelativeToEnd);
 	}
 	virtual bool Read(uint8* Destination, int64 BytesToRead) override
 	{
 		// Check that the file header is OK
-		if (!PakEntry.Verified)
+		if (!Reader.PakEntry.Verified)
 		{
 			FPakEntry FileHeader;
-			PakReader->Seek(PakEntry.Offset);
-			FileHeader.Serialize(*PakReader, PakFile.GetInfo().Version);
-			if (FPakEntry::VerifyPakEntriesMatch(PakEntry, FileHeader))
+			Reader.PakReader->Seek(Reader.PakEntry.Offset);
+			FileHeader.Serialize(*Reader.PakReader, Reader.PakFile.GetInfo().Version);
+			if (FPakEntry::VerifyPakEntriesMatch(Reader.PakEntry, FileHeader))
 			{
-				PakEntry.Verified = true;
+				Reader.PakEntry.Verified = true;
 			}
 			else
 			{
@@ -608,11 +745,10 @@ public:
 			}
 		}
 		//
-		PakReader->Seek(OffsetToFile + ReadPos);
-		if (PakEntry.Size >= (ReadPos + BytesToRead))
+		if (Reader.FileSize() >= (ReadPos + BytesToRead))
 		{
 			// Read directly from Pak.
-			PakReader->Serialize(Destination, BytesToRead);
+			Reader.Serialize(ReadPos, Destination, BytesToRead);
 			ReadPos += BytesToRead;
 			return true;
 		}
@@ -628,7 +764,7 @@ public:
 	}
 	virtual int64 Size() override
 	{
-		return PakEntry.Size;
+		return Reader.FileSize();
 	}
 	/// END IFileHandle Interface
 };
@@ -836,7 +972,7 @@ public:
 		const FPakEntry* FileEntry = FindFileInPakFiles(Filename);
 		if (FileEntry != NULL)
 		{
-			return FileEntry->Size;
+			return FileEntry->CompressionMethod != COMPRESS_None ? FileEntry->UncompressedSize : FileEntry->Size;
 		}
 		// First look for the file in the user dir.
 		int64 Result = LowerLevel->FileSize(Filename);
