Fix for UnrealPak parsing *nix style absolute paths as switches.

diff --git a/Engine/Source/Programs/UnrealPak/Private/UnrealPak.cpp b/Engine/Source/Programs/UnrealPak/Private/UnrealPak.cpp
--- a/Engine/Source/Programs/UnrealPak/Private/UnrealPak.cpp
+++ b/Engine/Source/Programs/UnrealPak/Private/UnrealPak.cpp
@@ -283,6 +283,22 @@ void ProcessOrderFile(int32 ArgC, TCHAR* ArgV[], TMap<FString, uint64>& OrderMap
 	}
 }
 
+static void CommandLineParseHelper(const TCHAR* InCmdLine, TArray<FString>& Tokens, TArray<FString>& Switches)
+{
+	FString NextToken;
+	while(FParse::Token(InCmdLine,NextToken,false))
+	{
+		if((**NextToken == TCHAR('-')))
+		{
+			new(Switches)FString(NextToken.Mid(1));
+		}
+		else
+		{
+			new(Tokens)FString(NextToken);
+		}
+	}
+}
+
 void ProcessCommandLine(int32 ArgC, TCHAR* ArgV[], TArray<FPakInputPair>& Entries)
 {
 	// List of all items to add to pak file
@@ -303,7 +319,7 @@ void ProcessCommandLine(int32 ArgC, TCHAR* ArgV[], TArray<FPakInputPair>& Entrie
 			{
 				TArray<FString> SourceAndDest;
 				TArray<FString> Switches;
-				FCommandLine::Parse(*Lines[EntryIndex], SourceAndDest, Switches);
+				CommandLineParseHelper(*Lines[EntryIndex], SourceAndDest, Switches);
 				FPakInputPair Input;
 				Input.Source = SourceAndDest[0];
 				FPaths::NormalizeFilename(Input.Source);
