Adding file system block alignment option in unreal pak.
Prevents small files crossing file system block which can incur a performance hit, this option should avoid that cost.
Changed UAT to automatically add this option to unreal pak on some platforms.

diff --git a/Engine/Source/Programs/AutomationTool/Platform.cs b/Engine/Source/Programs/AutomationTool/Platform.cs
--- a/Engine/Source/Programs/AutomationTool/Platform.cs
+++ b/Engine/Source/Programs/AutomationTool/Platform.cs
@@ -364,6 +364,14 @@ namespace AutomationTool
 			return false;
 		}
 
+        /// <summary>
+        /// Returns platform specific command line options for UnrealPak
+        /// </summary>
+        public virtual string GetPlatformPakCommandLine()
+        {
+            return "";
+        }
+
 		#region Hooks
 
 		public virtual void PreBuildAgenda(UE4Build Build, UE4Build.BuildAgenda Agenda)
diff --git a/Engine/Source/Programs/AutomationTool/Scripts/CopyBuildToStagingDirectory.Automation.cs b/Engine/Source/Programs/AutomationTool/Scripts/CopyBuildToStagingDirectory.Automation.cs
--- a/Engine/Source/Programs/AutomationTool/Scripts/CopyBuildToStagingDirectory.Automation.cs
+++ b/Engine/Source/Programs/AutomationTool/Scripts/CopyBuildToStagingDirectory.Automation.cs
@@ -53,7 +53,7 @@ public partial class Project : CommandUtils
 		return Result;
 	}
 
-	static public void RunUnrealPak(Dictionary<string, string> UnrealPakResponseFile, string OutputLocation, string EncryptionKeys, string PakOrderFileLocation, bool Compressed)
+	static public void RunUnrealPak(Dictionary<string, string> UnrealPakResponseFile, string OutputLocation, string EncryptionKeys, string PakOrderFileLocation, string PlatformOptions, bool Compressed)
 	{
 		if (UnrealPakResponseFile.Count < 1)
 		{
@@ -80,6 +80,7 @@ public partial class Project : CommandUtils
 		{
 			CmdLine += " -UTF8Output";
 		}
+        CmdLine += PlatformOptions;
 		RunAndLog(CmdEnv, UnrealPakExe, CmdLine, Options: ERunOptions.Default | ERunOptions.UTF8Output);
 		Log("UnrealPak Done *******");
 	}
@@ -407,7 +408,7 @@ public partial class Project : CommandUtils
 			PakOrderFileLocation = CombinePaths(PakOrderFileLocationBase, "EditorOpenOrder.log");
 		}
 
-		RunUnrealPak(UnrealPakResponseFile, OutputLocation, Params.SignPak, PakOrderFileLocation, Params.Compressed);
+        RunUnrealPak(UnrealPakResponseFile, OutputLocation, Params.SignPak, PakOrderFileLocation, SC.StageTargetPlatform.GetPlatformPakCommandLine(), Params.Compressed);
 
 		// add the pak file as needing deployment and convert to lower case again if needed
 		SC.UFSStagingFiles.Add(OutputLocation, OutputRealtiveLocation);		
diff --git a/Engine/Source/Programs/UnrealPak/Private/UnrealPak.cpp b/Engine/Source/Programs/UnrealPak/Private/UnrealPak.cpp
--- a/Engine/Source/Programs/UnrealPak/Private/UnrealPak.cpp
+++ b/Engine/Source/Programs/UnrealPak/Private/UnrealPak.cpp
@@ -8,9 +8,21 @@
 #include "SignedArchiveWriter.h"
 #include "KeyGenerator.h"
 #include "AES.h"
+#include "UniquePtr.h"
 
 IMPLEMENT_APPLICATION(UnrealPak, "UnrealPak");
 
+struct FPakCommandLineParameters
+{
+	FPakCommandLineParameters()
+		: CompressionBlockSize(64*1024)
+		, FileSystemBlockSize(0)
+	{}
+
+	int32  CompressionBlockSize;
+	int64  FileSystemBlockSize;
+};
+
 struct FPakEntryPair
 {
 	FString Filename;
@@ -49,6 +61,49 @@ struct FPakEntryOrder
 	uint64  Order;
 };
 
+struct FCompressedFileBuffer
+{
+	FCompressedFileBuffer()
+		: OriginalSize(0)
+		,TotalCompressedSize(0)
+		,FileCompressionBlockSize(0)
+		,CompressedBufferSize(0)
+	{
+
+	}
+
+	void Reinitialize(FArchive* File,ECompressionFlags CompressionMethod,int64 CompressionBlockSize)
+	{
+		OriginalSize = File->TotalSize();
+		TotalCompressedSize = 0;
+		FileCompressionBlockSize = 0;
+		FileCompressionMethod = CompressionMethod;
+		CompressedBlocks.Reset();
+		CompressedBlocks.AddUninitialized((OriginalSize+CompressionBlockSize-1)/CompressionBlockSize);
+	}
+
+	void EnsureBufferSpace(int64 RequiredSpace)
+	{
+		if(RequiredSpace > CompressedBufferSize)
+		{
+			uint8* NewPtr = (uint8*)FMemory::Malloc(RequiredSpace);
+			FMemory::Memcpy(NewPtr,CompressedBuffer.Get(),CompressedBufferSize);
+			CompressedBuffer.Reset(NewPtr);
+			CompressedBufferSize = RequiredSpace;
+		}
+	}
+
+	bool CompressFileToWorkingBuffer(const FPakInputPair& InFile,uint8*& InOutPersistentBuffer,int64& InOutBufferSize,ECompressionFlags CompressionMethod,const int32 CompressionBlockSize);
+
+	int64						 OriginalSize;
+	int64						 TotalCompressedSize;
+	int32						 FileCompressionBlockSize;
+	ECompressionFlags			 FileCompressionMethod;
+	TArray<FPakCompressedBlock>  CompressedBlocks;
+	int64						 CompressedBufferSize;
+	TUniquePtr<uint8>		     CompressedBuffer;
+};
+
 FString GetLongestPath(TArray<FPakInputPair>& FilesToAdd)
 {
 	FString LongestPath;
@@ -107,8 +162,71 @@ FString GetCommonRootPath(TArray<FPakInputPair>& FilesToAdd)
 	return Root;
 }
 
+bool FCompressedFileBuffer::CompressFileToWorkingBuffer(const FPakInputPair& InFile, uint8*& InOutPersistentBuffer, int64& InOutBufferSize, ECompressionFlags CompressionMethod, const int32 CompressionBlockSize)
+{
+	TAutoPtr<FArchive> FileHandle(IFileManager::Get().CreateFileReader(*InFile.Source));
+	if(!FileHandle.IsValid())
+	{
+		TotalCompressedSize = 0;
+		return false;
+	}
 
-bool CopyFileToPak(FArchive& InPak, const FString& InMountPoint, const FPakInputPair& InFile, uint8*& InOutPersistentBuffer, int64& InOutBufferSize, uint8*& InPersistentCompressionBuffer, int64& InBufferCompressionSize, FPakEntryPair& OutNewEntry, ECompressionFlags CompressionMethod, const int32 CompressionBlockSize)
+	Reinitialize(FileHandle.GetOwnedPointer(), CompressionMethod, CompressionBlockSize);
+	const int64 FileSize = OriginalSize;
+	const int64 PaddedEncryptedFileSize = Align(FileSize,FAES::AESBlockSize);
+	if(InOutBufferSize < PaddedEncryptedFileSize)
+	{
+		InOutPersistentBuffer = (uint8*)FMemory::Realloc(InOutPersistentBuffer,PaddedEncryptedFileSize);
+		InOutBufferSize = FileSize;
+	}
+
+	// Load to buffer
+	FileHandle->Serialize(InOutPersistentBuffer,FileSize);
+
+	// Build buffers for working
+	int64 UncompressedSize = FileSize;
+	int32 CompressionBufferSize = Align(FCompression::CompressMemoryBound(CompressionMethod,CompressionBlockSize),FAES::AESBlockSize);
+	EnsureBufferSpace(Align(FCompression::CompressMemoryBound(CompressionMethod,FileSize),FAES::AESBlockSize));
+
+
+	TotalCompressedSize = 0;
+	int64 UncompressedBytes = 0;
+	int32 CurrentBlock = 0;
+	while(UncompressedSize)
+	{
+		int32 BlockSize = (int32)FMath::Min<int64>(UncompressedSize,CompressionBlockSize);
+		int32 CompressedBlockSize = CompressionBufferSize;
+		FileCompressionBlockSize = FMath::Max<uint32>(BlockSize, FileCompressionBlockSize);
+		EnsureBufferSpace(TotalCompressedSize+CompressedBlockSize);
+		if(!FCompression::CompressMemory(CompressionMethod,CompressedBuffer.Get()+TotalCompressedSize,CompressedBlockSize,InOutPersistentBuffer+UncompressedBytes,BlockSize))
+		{
+			return false;
+		}
+		UncompressedSize -= BlockSize;
+		UncompressedBytes += BlockSize;
+
+		CompressedBlocks[CurrentBlock].CompressedStart = TotalCompressedSize;
+		CompressedBlocks[CurrentBlock].CompressedEnd = TotalCompressedSize+CompressedBlockSize;
+		++CurrentBlock;
+
+		TotalCompressedSize += CompressedBlockSize;
+
+		if(InFile.bNeedEncryption)
+		{
+			int32 EncryptionBlockPadding = Align(TotalCompressedSize,FAES::AESBlockSize);
+			for(int64 FillIndex=TotalCompressedSize; FillIndex < EncryptionBlockPadding; ++FillIndex)
+			{
+				// Fill the trailing buffer with random bytes from file
+				CompressedBuffer.Get()[FillIndex] = CompressedBuffer.Get()[rand()%TotalCompressedSize];
+			}
+			TotalCompressedSize += EncryptionBlockPadding - TotalCompressedSize;
+		}
+	}
+
+	return true;
+}
+
+bool CopyFileToPak(FArchive& InPak, const FString& InMountPoint, const FPakInputPair& InFile, uint8*& InOutPersistentBuffer, int64& InOutBufferSize, FPakEntryPair& OutNewEntry)
 {	
 	TAutoPtr<FArchive> FileHandle(IFileManager::Get().CreateFileReader(*InFile.Source));
 	bool bFileExists = FileHandle.IsValid();
@@ -120,7 +238,7 @@ bool CopyFileToPak(FArchive& InPak, const FString& InMountPoint, const FPakInput
 		OutNewEntry.Info.Offset = 0; // Don't serialize offsets here.
 		OutNewEntry.Info.Size = FileSize;
 		OutNewEntry.Info.UncompressedSize = FileSize;
-		OutNewEntry.Info.CompressionMethod = CompressionMethod;
+		OutNewEntry.Info.CompressionMethod = COMPRESS_None;
 		OutNewEntry.Info.bEncrypted = InFile.bNeedEncryption;
 
 		if (InOutBufferSize < PaddedEncryptedFileSize)
@@ -129,16 +247,9 @@ bool CopyFileToPak(FArchive& InPak, const FString& InMountPoint, const FPakInput
 			InOutBufferSize = FileSize;
 		}
 
-		// Don't store empty files with compression, it's wasteful
-		if (FileSize == 0) 
-		{
-			CompressionMethod = COMPRESS_None;
-		}
-
 		// Load to buffer
 		FileHandle->Serialize(InOutPersistentBuffer, FileSize);
 
-		if (CompressionMethod == COMPRESS_None)
 		{
 			int64 SizeToWrite = FileSize;
 			if (InFile.bNeedEncryption)
@@ -163,82 +274,53 @@ bool CopyFileToPak(FArchive& InPak, const FString& InMountPoint, const FPakInput
 			OutNewEntry.Info.Serialize(InPak,FPakInfo::PakFile_Version_Latest);
 			InPak.Serialize(InOutPersistentBuffer,SizeToWrite);
 		}
-		else
-		{
-			// Build buffers for working
-			int64 UncompressedSize = FileSize;
-			int32 CompressionBufferSize = Align(FCompression::CompressMemoryBound(CompressionMethod, CompressionBlockSize), FAES::AESBlockSize);
-			if (CompressionBufferSize > InBufferCompressionSize)
-			{
-				InPersistentCompressionBuffer = (uint8*)FMemory::Realloc(InPersistentCompressionBuffer,CompressionBufferSize);
-				InBufferCompressionSize = CompressionBufferSize;
-			}
-			
-			int64 HeaderTell = InPak.Tell();
-			OutNewEntry.Info.CompressionBlocks.AddUninitialized((FileSize+CompressionBlockSize-1)/CompressionBlockSize);
-			OutNewEntry.Info.Serialize(InPak,FPakInfo::PakFile_Version_Latest);
+	}
+	return bFileExists;
+}
 
-			int64 UncompressedBytes = 0;
-			int64 CompressedBytes = 0;
-			int32 CurrentBlock = 0;
-			FSHA1 Hash;
-			Hash.Reset();
-			while (UncompressedSize)
-			{
-				int32 BlockSize = (int32)FMath::Min<int64>(UncompressedSize, CompressionBlockSize);
-				int32 CompressedBlockSize = CompressionBufferSize;
-				OutNewEntry.Info.CompressionBlockSize = FMath::Max<uint32>(BlockSize, OutNewEntry.Info.CompressionBlockSize);
-				if (!FCompression::CompressMemory(CompressionMethod, InPersistentCompressionBuffer, CompressedBlockSize, InOutPersistentBuffer+UncompressedBytes, BlockSize))
-				{
-					return false;
-				}
-				UncompressedSize -= BlockSize;
-				UncompressedBytes += BlockSize;
-				CompressedBytes += CompressedBlockSize;
+bool CopyCompressedFileToPak(FArchive& InPak, const FString& InMountPoint, const FPakInputPair& InFile, const FCompressedFileBuffer& CompressedFile, FPakEntryPair& OutNewEntry)
+{
+	if (CompressedFile.TotalCompressedSize == 0)
+	{
+		return false;
+	}
 
-				OutNewEntry.Info.CompressionBlocks[CurrentBlock].CompressedStart = InPak.Tell();
-				OutNewEntry.Info.CompressionBlocks[CurrentBlock].CompressedEnd = InPak.Tell()+CompressedBlockSize;
-				++CurrentBlock;
+	int64 HeaderTell = InPak.Tell();
+	OutNewEntry.Info.CompressionMethod = CompressedFile.FileCompressionMethod;
+	OutNewEntry.Info.CompressionBlocks.AddUninitialized(CompressedFile.CompressedBlocks.Num());
 
-				if (InFile.bNeedEncryption)
-				{
-					int32 EncryptionBlockPadding = Align(CompressedBlockSize,FAES::AESBlockSize);
-					for(int64 FillIndex=CompressedBlockSize; FillIndex < EncryptionBlockPadding; ++FillIndex)
-					{
-						// Fill the trailing buffer with random bytes from file
-						InPersistentCompressionBuffer[FillIndex] = InPersistentCompressionBuffer[rand()%CompressedBlockSize];
-					}
-					FAES::EncryptData(InPersistentCompressionBuffer, EncryptionBlockPadding);
-					CompressedBytes += EncryptionBlockPadding - CompressedBlockSize;
+	int64 TellPos = InPak.Tell() + OutNewEntry.Info.GetSerializedSize(FPakInfo::PakFile_Version_Latest);
+	const TArray<FPakCompressedBlock>& Blocks = CompressedFile.CompressedBlocks;
+	for (int32 BlockIndex = 0, BlockCount = CompressedFile.CompressedBlocks.Num(); BlockIndex < BlockCount; ++BlockIndex)
+	{
+		OutNewEntry.Info.CompressionBlocks[BlockIndex].CompressedStart = Blocks[BlockIndex].CompressedStart + TellPos;
+		OutNewEntry.Info.CompressionBlocks[BlockIndex].CompressedEnd = Blocks[BlockIndex].CompressedEnd + TellPos;
+	}
 
-					//Write out buffer
-					InPak.Serialize(InPersistentCompressionBuffer,EncryptionBlockPadding);
-					// Calculate the buffer hash value
-					Hash.Update(InPersistentCompressionBuffer,EncryptionBlockPadding);
-				}
-				else
-				{
-					//Write out buffer
-					InPak.Serialize(InPersistentCompressionBuffer, CompressedBlockSize);
-					// Calculate the buffer hash value
-					Hash.Update(InPersistentCompressionBuffer, CompressedBlockSize);
-				}
-			}
+	if (InFile.bNeedEncryption)
+	{
+		FAES::EncryptData(CompressedFile.CompressedBuffer.Get(), CompressedFile.TotalCompressedSize);
+	}
 
-			Hash.Final();
+	//Hash the final buffer thats written
+	FSHA1 Hash;
+	Hash.Update(CompressedFile.CompressedBuffer.Get(), CompressedFile.TotalCompressedSize);
+	Hash.Final();
 
-			// Update file size & Hash
-			OutNewEntry.Info.Size = CompressedBytes;
-			Hash.GetHash(OutNewEntry.Info.Hash);
+	// Update file size & Hash
+	OutNewEntry.Info.CompressionBlockSize = CompressedFile.FileCompressionBlockSize;
+	OutNewEntry.Info.UncompressedSize = CompressedFile.OriginalSize;
+	OutNewEntry.Info.Size = CompressedFile.TotalCompressedSize;
+	Hash.GetHash(OutNewEntry.Info.Hash);
 
-			// Write out the corrected header
-			int64 CurrentTell = InPak.Tell();
-			InPak.Seek(HeaderTell);
-			OutNewEntry.Info.Serialize(InPak,FPakInfo::PakFile_Version_Latest);
-			InPak.Seek(CurrentTell);
-		}
-	}
-	return bFileExists;
+	//	Write the header, then the data
+	OutNewEntry.Filename = InFile.Dest.Mid(InMountPoint.Len());
+	OutNewEntry.Info.Offset = 0; // Don't serialize offsets here.
+	OutNewEntry.Info.bEncrypted = InFile.bNeedEncryption;
+	OutNewEntry.Info.Serialize(InPak,FPakInfo::PakFile_Version_Latest);
+	InPak.Serialize(CompressedFile.CompressedBuffer.Get(), CompressedFile.TotalCompressedSize);
+
+	return true;
 }
 
 void ProcessOrderFile(int32 ArgC, TCHAR* ArgV[], TMap<FString, uint64>& OrderMap)
@@ -299,10 +381,28 @@ static void CommandLineParseHelper(const TCHAR* InCmdLine, TArray<FString>& Toke
 	}
 }
 
-void ProcessCommandLine(int32 ArgC, TCHAR* ArgV[], TArray<FPakInputPair>& Entries)
+void ProcessCommandLine(int32 ArgC, TCHAR* ArgV[], TArray<FPakInputPair>& Entries, FPakCommandLineParameters& CmdLineParameters)
 {
 	// List of all items to add to pak file
 	FString ResponseFile;
+	FString ClusterSizeString;
+
+	if (FParse::Value(FCommandLine::Get(), TEXT("-blocksize="), ClusterSizeString) && 
+		FParse::Value(FCommandLine::Get(), TEXT("-blocksize="), CmdLineParameters.FileSystemBlockSize))
+	{
+		if (ClusterSizeString.EndsWith(TEXT("MB")))
+		{
+			CmdLineParameters.FileSystemBlockSize *= 1024*1024;
+		}
+		else if (ClusterSizeString.EndsWith(TEXT("KB")))
+		{
+			CmdLineParameters.FileSystemBlockSize *= 1024;
+		}
+	}
+	else
+	{
+		CmdLineParameters.FileSystemBlockSize = 0;
+	}
 
 	if (FParse::Value(FCommandLine::Get(), TEXT("-create="), ResponseFile))
 	{
@@ -603,7 +703,7 @@ FArchive* CreatePakWriter(const TCHAR* Filename)
 	return Writer;
 }
 
-bool CreatePakFile(const TCHAR* Filename, TArray<FPakInputPair>& FilesToAdd, const int32 CompressionBlockSize)
+bool CreatePakFile(const TCHAR* Filename, TArray<FPakInputPair>& FilesToAdd, const FPakCommandLineParameters& CmdLineParameters)
 {	
 	const double StartTime = FPlatformTime::Seconds();
 
@@ -620,20 +720,53 @@ bool CreatePakFile(const TCHAR* Filename, TArray<FPakInputPair>& FilesToAdd, con
 	FString MountPoint = GetCommonRootPath(FilesToAdd);
 	uint8* ReadBuffer = NULL;
 	int64 BufferSize = 0;
-	uint8* CompressionBuffer = NULL;
-	int64 CompressionBufferSize = 0;
 	ECompressionFlags CompressionMethod = COMPRESS_None;
+	FCompressedFileBuffer CompressedFileBuffer;
 
 	for (int32 FileIndex = 0; FileIndex < FilesToAdd.Num(); FileIndex++)
 	{
 		//  Remember the offset but don't serialize it with the entry header.
-		const int64 NewEntryOffset = PakFileHandle->Tell();
+		int64 NewEntryOffset = PakFileHandle->Tell();
 		FPakEntryPair NewEntry;
 
 		//check if this file requested to be compression
-		CompressionMethod = FilesToAdd[FileIndex].bNeedsCompression ? COMPRESS_Default : COMPRESS_None;
+		int64 OriginalFileSize = IFileManager::Get().FileSize(*FilesToAdd[FileIndex].Source);
+		int64 RealFileSize = OriginalFileSize + NewEntry.Info.GetSerializedSize(FPakInfo::PakFile_Version_Latest);
+		CompressionMethod = (FilesToAdd[FileIndex].bNeedsCompression && OriginalFileSize > 0) ? COMPRESS_Default : COMPRESS_None;
+
+		if (CompressionMethod != COMPRESS_None)
+		{
+			if (CompressedFileBuffer.CompressFileToWorkingBuffer(FilesToAdd[FileIndex], ReadBuffer, BufferSize, CompressionMethod, CmdLineParameters.CompressionBlockSize))
+			{
+				NewEntry.Info.CompressionMethod = CompressionMethod;
+				NewEntry.Info.CompressionBlocks.AddUninitialized(CompressedFileBuffer.CompressedBlocks.Num());
+				RealFileSize = CompressedFileBuffer.TotalCompressedSize + NewEntry.Info.GetSerializedSize(FPakInfo::PakFile_Version_Latest);
+				NewEntry.Info.CompressionBlocks.Reset();
+			}
+		}
+
+		// Account for file system block size, which is a boundary we want to avoid crossing.
+		if (CmdLineParameters.FileSystemBlockSize > 0 && OriginalFileSize != INDEX_NONE && RealFileSize <= CmdLineParameters.FileSystemBlockSize)
+		{
+			if ((NewEntryOffset / CmdLineParameters.FileSystemBlockSize) != ((NewEntryOffset+RealFileSize) / CmdLineParameters.FileSystemBlockSize))
+			{
+				//File crosses a block boundary, so align it to the beginning of the next boundary
+				NewEntryOffset = AlignArbitrary(NewEntryOffset, CmdLineParameters.FileSystemBlockSize);
+				PakFileHandle->Seek(NewEntryOffset);
+			}
+		}
+
+		bool bCopiedToPak;
+		if (FilesToAdd[FileIndex].bNeedsCompression && CompressionMethod != COMPRESS_None)
+		{
+			bCopiedToPak = CopyCompressedFileToPak(*PakFileHandle, MountPoint, FilesToAdd[FileIndex], CompressedFileBuffer, NewEntry);
+		}
+		else
+		{
+			bCopiedToPak = CopyFileToPak(*PakFileHandle, MountPoint, FilesToAdd[FileIndex], ReadBuffer, BufferSize, NewEntry);
+		}
 
-		if (CopyFileToPak(*PakFileHandle, MountPoint, FilesToAdd[FileIndex], ReadBuffer, BufferSize, CompressionBuffer, CompressionBufferSize, NewEntry, CompressionMethod, CompressionBlockSize) == true)
+		if (bCopiedToPak)
 		{
 			// Update offset now and store it in the index (and only in index)
 			NewEntry.Info.Offset = NewEntryOffset;
@@ -835,7 +968,7 @@ INT32_MAIN_INT32_ARGC_TCHAR_ARGV()
 		return 1;
 	}
 
-	int32 CompressionBlockSize = 64*1024;
+	FPakCommandLineParameters CmdLineParameters;
 	int32 Result = 0;
 	FString KeyFilename;
 	if (FParse::Value(FCommandLine::Get(), TEXT("GenerateKeys="), KeyFilename, false))
@@ -874,7 +1007,7 @@ INT32_MAIN_INT32_ARGC_TCHAR_ARGV()
 		{
 			// List of all items to add to pak file
 			TArray<FPakInputPair> Entries;
-			ProcessCommandLine(ArgC, ArgV, Entries);
+			ProcessCommandLine(ArgC, ArgV, Entries, CmdLineParameters);
 			TMap<FString, uint64> OrderMap;
 			ProcessOrderFile(ArgC, ArgV, OrderMap);
 
@@ -889,7 +1022,7 @@ INT32_MAIN_INT32_ARGC_TCHAR_ARGV()
 				TArray<FPakInputPair> FilesToAdd;
 				CollectFilesToAdd(FilesToAdd, Entries, OrderMap);
 
-				Result = CreatePakFile(*PakFilename, FilesToAdd, CompressionBlockSize) ? 0 : 1;
+				Result = CreatePakFile(*PakFilename, FilesToAdd, CmdLineParameters) ? 0 : 1;
 			}
 		}
 	}
