Allow specification of the file list via a directory name intead of solely a response file
Allow compress and encrypt to be specified globally for all files via the commandline

diff --git a/Engine/Source/Programs/UnrealPak/Private/UnrealPak.cpp b/Engine/Source/Programs/UnrealPak/Private/UnrealPak.cpp
--- a/Engine/Source/Programs/UnrealPak/Private/UnrealPak.cpp
+++ b/Engine/Source/Programs/UnrealPak/Private/UnrealPak.cpp
@@ -406,55 +406,78 @@ void ProcessCommandLine(int32 ArgC, TCHAR* ArgV[], TArray<FPakInputPair>& Entrie
 
 	if (FParse::Value(FCommandLine::Get(), TEXT("-create="), ResponseFile))
 	{
-		FString Text;
-		UE_LOG(LogPakFile, Display, TEXT("Loading response file %s"), *ResponseFile);
-		if (FFileHelper::LoadFileToString(Text, *ResponseFile))
+		bool bCompress = false;
+		bool bEncrypt = false;
+		TArray<FString> Lines;
+
+		if (FParse::Param(FCommandLine::Get(), TEXT("compress")))
 		{
-			// Remove all carriage return characters.
-			Text.ReplaceInline(TEXT("\r"), TEXT(""));
-			// Read all lines
-			TArray<FString> Lines;
-			Text.ParseIntoArray(&Lines, TEXT("\n"), true);
-			for (int32 EntryIndex = 0; EntryIndex < Lines.Num(); EntryIndex++)
+			bCompress = true;
+		}
+		if (FParse::Param(FCommandLine::Get(), TEXT("encrypt")))
+		{
+			bEncrypt = true;
+		}
+
+		if (IFileManager::Get().DirectoryExists(*ResponseFile))
+		{
+			IFileManager::Get().FindFilesRecursive(Lines, *ResponseFile, TEXT("*"), true, false);
+		}
+		else
+		{
+			FString Text;
+			UE_LOG(LogPakFile, Display, TEXT("Loading response file %s"), *ResponseFile);
+			if (FFileHelper::LoadFileToString(Text, *ResponseFile))
 			{
-				TArray<FString> SourceAndDest;
-				TArray<FString> Switches;
-				CommandLineParseHelper(*Lines[EntryIndex], SourceAndDest, Switches);
-				FPakInputPair Input;
-				Input.Source = SourceAndDest[0];
-				FPaths::NormalizeFilename(Input.Source);
-				if (SourceAndDest.Num() > 1)
-				{
-					Input.Dest = FPaths::GetPath(SourceAndDest[1]);
-				}
-				else
+				// Remove all carriage return characters.
+				Text.ReplaceInline(TEXT("\r"), TEXT(""));
+				// Read all lines
+				Text.ParseIntoArray(&Lines, TEXT("\n"), true);
+			}
+			else
+			{
+				UE_LOG(LogPakFile, Error, TEXT("Failed to load %s"), *ResponseFile);
+			}
+		}
+
+		for (int32 EntryIndex = 0; EntryIndex < Lines.Num(); EntryIndex++)
+		{
+			TArray<FString> SourceAndDest;
+			TArray<FString> Switches;
+			CommandLineParseHelper(*Lines[EntryIndex], SourceAndDest, Switches);
+			FPakInputPair Input;
+
+			Input.Source = SourceAndDest[0];
+			FPaths::NormalizeFilename(Input.Source);
+			if (SourceAndDest.Num() > 1)
+			{
+				Input.Dest = FPaths::GetPath(SourceAndDest[1]);
+			}
+			else
+			{
+				Input.Dest = FPaths::GetPath(Input.Source);
+			}
+			FPaths::NormalizeFilename(Input.Dest);
+			FPakFile::MakeDirectoryFromPath(Input.Dest);
+
+			//check for compression switches
+			for (int32 Index = 0; Index < Switches.Num(); ++Index)
+			{
+				if (Switches[Index] == TEXT("compress"))
 				{
-					Input.Dest = FPaths::GetPath(Input.Source);
+					Input.bNeedsCompression = true;
 				}
-				FPaths::NormalizeFilename(Input.Dest);
-				FPakFile::MakeDirectoryFromPath(Input.Dest);
-
-				//check for compression switches
-				for (int32 Index = 0; Index < Switches.Num(); ++Index)
+				if (Switches[Index] == TEXT("encrypt"))
 				{
-					if (Switches[Index] == TEXT("compress"))
-					{
-						Input.bNeedsCompression = true;
-					}
-					if(Switches[Index] == TEXT("encrypt"))
-					{
-						Input.bNeedEncryption = true;
-					}
+					Input.bNeedEncryption = true;
 				}
+			}
+			Input.bNeedsCompression |= bCompress;
+			Input.bNeedEncryption |= bEncrypt;
 
-				UE_LOG(LogPakFile, Display, TEXT("Added file Source: %s Dest: %s"), *Input.Source, *Input.Dest);
-				Entries.Add(Input);
-			}			
-		}
-		else
-		{
-			UE_LOG(LogPakFile, Error, TEXT("Failed to load %s"), *ResponseFile);
-		}
+			UE_LOG(LogPakFile, Display, TEXT("Added file Source: %s Dest: %s"), *Input.Source, *Input.Dest);
+			Entries.Add(Input);
+		}			
 	}
 	else
 	{
