Speed up CollectFilesToAdd() by moving the hash table look ups out of the sorting function.

diff --git a/Engine/Source/Programs/UnrealPak/Private/UnrealPak.cpp b/Engine/Source/Programs/UnrealPak/Private/UnrealPak.cpp
--- a/Engine/Source/Programs/UnrealPak/Private/UnrealPak.cpp
+++ b/Engine/Source/Programs/UnrealPak/Private/UnrealPak.cpp
@@ -33,11 +33,13 @@ struct FPakInputPair
 {
 	FString Source;
 	FString Dest;
+	uint64 SuggestedOrder; 
 	bool bNeedsCompression;
 	bool bNeedEncryption;
 
 	FPakInputPair()
-		: bNeedsCompression(false)
+		: SuggestedOrder(MAX_uint64)
+		, bNeedsCompression(false)
 		, bNeedEncryption(false)
 	{}
 
@@ -552,6 +554,11 @@ void CollectFilesToAdd(TArray<FPakInputPair>& OutFilesToAdd, const TArray<FPakIn
 				FileInput.Source = FoundFiles[FileIndex];
 				FPaths::MakeStandardFilename(FileInput.Source);
 				FileInput.Dest = FileInput.Source.Replace(*Directory, *Input.Dest, ESearchCase::IgnoreCase);
+				const uint64* FoundOrder = OrderMap.Find(FileInput.Dest);
+				if(FoundOrder)
+				{
+					FileInput.SuggestedOrder = *FoundOrder;
+				}
 				FileInput.bNeedsCompression = bCompression;
 				FileInput.bNeedEncryption = bEncryption;
 				if (!AddedFiles.Contains(FileInput.Source))
@@ -575,6 +582,11 @@ void CollectFilesToAdd(TArray<FPakInputPair>& OutFilesToAdd, const TArray<FPakIn
 			FileInput.Source = Input.Source;
 			FPaths::MakeStandardFilename(FileInput.Source);
 			FileInput.Dest = FileInput.Source.Replace(*Directory, *Input.Dest, ESearchCase::IgnoreCase);
+			const uint64* FoundOrder = OrderMap.Find(FileInput.Dest);
+			if (FoundOrder)
+			{
+				FileInput.SuggestedOrder = *FoundOrder;
+			}
 			FileInput.bNeedEncryption = bEncryption;
 			FileInput.bNeedsCompression = bCompression;
 
@@ -592,26 +604,15 @@ void CollectFilesToAdd(TArray<FPakInputPair>& OutFilesToAdd, const TArray<FPakIn
 		}
 	}
 
-	// Sort alphabetically
+	// Sort by suggested order then alphabetically
 	struct FInputPairSort
 	{
-		const TMap<FString, uint64>& OrderMap;
-
-		FInputPairSort(const TMap<FString, uint64>& InOrderMap) : OrderMap(InOrderMap) {}
-		
 		FORCEINLINE bool operator()(const FPakInputPair& A, const FPakInputPair& B) const
 		{
-			uint64 AOrder = MAX_uint64;
-			uint64 BOrder = MAX_uint64;
-			const uint64* FoundOrder;
-			FoundOrder = OrderMap.Find(A.Dest);
-			if (FoundOrder) AOrder = *FoundOrder;
-			FoundOrder = OrderMap.Find(B.Dest);
-			if (FoundOrder) BOrder = *FoundOrder;
-			return AOrder == BOrder ? A.Dest < B.Dest : AOrder < BOrder;
+			return A.SuggestedOrder == B.SuggestedOrder ? A.Dest < B.Dest : A.SuggestedOrder < B.SuggestedOrder;
 		}
 	};
-	OutFilesToAdd.Sort(FInputPairSort(OrderMap));
+	OutFilesToAdd.Sort(FInputPairSort());
 	UE_LOG(LogPakFile, Display, TEXT("Collected %d files in %.2lfs."), OutFilesToAdd.Num(), FPlatformTime::Seconds() - StartTime);
 }
 
