Added -List option to list all files in a pak file with stored size to UnrealPak

diff --git a/Engine/Source/Programs/UnrealPak/Private/UnrealPak.cpp b/Engine/Source/Programs/UnrealPak/Private/UnrealPak.cpp
--- a/Engine/Source/Programs/UnrealPak/Private/UnrealPak.cpp
+++ b/Engine/Source/Programs/UnrealPak/Private/UnrealPak.cpp
@@ -901,6 +901,32 @@ bool TestPakFile(const TCHAR* Filename)
 	}
 }
 
+bool ListFilesInPak(const TCHAR * InPakFilename)
+{
+	FPakFile PakFile(InPakFilename, FParse::Param(FCommandLine::Get(), TEXT("signed")));
+	int32 FileCount = 0;
+	int64 FileSize = 0;
+
+	if (PakFile.IsValid())
+	{
+		for (FPakFile::FFileIterator It(PakFile); It; ++It, ++FileCount)
+		{
+			const FPakEntry& Entry = It.Info();
+			UE_LOG(LogPakFile, Display, TEXT("\"%s\" %d bytes."), *It.Filename(), Entry.Size);
+			FileSize += Entry.Size;
+			FileCount++;
+		}
+		UE_LOG(LogPakFile, Display, TEXT("%d files (%lld bytes)."), FileCount, FileSize);
+
+		return true;
+	}
+	else
+	{
+		UE_LOG(LogPakFile, Error, TEXT("Unable to open pak file \"%s\"."), InPakFilename);
+		return false;
+	}
+}
+
 bool ExtractFilesFromPak(const TCHAR* InPakFilename, const TCHAR* InDestPath)
 {
 	FPakFile PakFile(InPakFilename, FParse::Param(FCommandLine::Get(), TEXT("signed")));
@@ -1015,6 +1041,10 @@ INT32_MAIN_INT32_ARGC_TCHAR_ARGV()
 		{
 			Result = TestPakFile(*PakFilename) ? 0 : 1;
 		}
+		else if (FParse::Param(FCommandLine::Get(), TEXT("List")))
+		{
+			Result = ListFilesInPak(*PakFilename);
+		}
 		else if (FParse::Param(FCommandLine::Get(), TEXT("Extract")))
 		{
 			if (ArgC < 4)
