Fixed issue with blank lines in feature pack response file import

diff --git a/Engine/Source/Programs/UnrealPak/Private/UnrealPak.cpp b/Engine/Source/Programs/UnrealPak/Private/UnrealPak.cpp
--- a/Engine/Source/Programs/UnrealPak/Private/UnrealPak.cpp
+++ b/Engine/Source/Programs/UnrealPak/Private/UnrealPak.cpp
@@ -446,7 +446,11 @@ void ProcessCommandLine(int32 ArgC, TCHAR* ArgV[], TArray<FPakInputPair>& Entrie
 		{
 			TArray<FString> SourceAndDest;
 			TArray<FString> Switches;
-			CommandLineParseHelper(*Lines[EntryIndex], SourceAndDest, Switches);
+			CommandLineParseHelper(*Lines[EntryIndex].Trim(), SourceAndDest, Switches);
+			if( SourceAndDest.Num() == 0)
+			{
+				continue;
+			}
 			FPakInputPair Input;
 
 			Input.Source = SourceAndDest[0];
