- Added FPaths::LaunchDir(), which gives the current working directory _at launch time_ (saved off before we switch to the Binaries dir). This can be useful for finding files relative to working dir for commandline utilities
- Modified UnrealPak to use LaunchDir as a fallback to finding a pak file
- Added "Usage" printout to UnrealPak that has all the use cases and options I could find in the code

diff --git a/Engine/Source/Programs/UnrealPak/Private/UnrealPak.cpp b/Engine/Source/Programs/UnrealPak/Private/UnrealPak.cpp
--- a/Engine/Source/Programs/UnrealPak/Private/UnrealPak.cpp
+++ b/Engine/Source/Programs/UnrealPak/Private/UnrealPak.cpp
@@ -1107,10 +1107,22 @@ INT32_MAIN_INT32_ARGC_TCHAR_ARGV()
 {
 	// start up the main loop
 	GEngineLoop.PreInit(ArgC, ArgV);
-	
+
 	if (ArgC < 2)
 	{
-		UE_LOG(LogPakFile, Error, TEXT("No pak file name specified."));
+		UE_LOG(LogPakFile, Error, TEXT("No pak file name specified. Usage:"));
+		UE_LOG(LogPakFile, Error, TEXT("  UnrealPak <PakFilename> -Test"));
+		UE_LOG(LogPakFile, Error, TEXT("  UnrealPak <PakFilename> -List"));
+		UE_LOG(LogPakFile, Error, TEXT("  UnrealPak <PakFilename> -Extract <ExtractDir>"));
+		UE_LOG(LogPakFile, Error, TEXT("  UnrealPak <PakFilename> -Create=<ResponseFile> [Options]"));
+		UE_LOG(LogPakFile, Error, TEXT("  UnrealPak <PakFilename> -Dest=<MountPoint>"));
+		UE_LOG(LogPakFile, Error, TEXT("  UnrealPak GenerateKeys=<KeyFilename>"));
+		UE_LOG(LogPakFile, Error, TEXT("  UnrealPak GeneratePrimeTable=<KeyFilename> [-TableMax=<N>]"));
+		UE_LOG(LogPakFile, Error, TEXT("  Options:"));
+		UE_LOG(LogPakFile, Error, TEXT("    -blocksize=<BlockSize>"));
+		UE_LOG(LogPakFile, Error, TEXT("    -compress"));
+		UE_LOG(LogPakFile, Error, TEXT("    -encrypt"));
+		UE_LOG(LogPakFile, Error, TEXT("    -order=<OrderingFile>"));
 		return 1;
 	}
 
@@ -1131,6 +1143,18 @@ INT32_MAIN_INT32_ARGC_TCHAR_ARGV()
 	{
 		FString PakFilename(ArgV[1]);
 		FPaths::MakeStandardFilename(PakFilename);
+		
+		// make sure it exists relative to engine, or relative to launch dir
+		if (!FPaths::FileExists(PakFilename))
+		{
+			PakFilename = FPaths::LaunchDir() + ArgV[1];
+		}
+			
+		if (!FPaths::FileExists(PakFilename))
+		{
+			UE_LOG(LogPakFile, Error, TEXT("Pak file %s could not be found (checked against binary and launch directories)"), ArgV[1]);
+			return 1;
+		}
 
 		if (FParse::Param(FCommandLine::Get(), TEXT("Test")))
 		{
diff --git a/Engine/Source/Runtime/Core/Private/GenericPlatform/GenericPlatformMisc.cpp b/Engine/Source/Runtime/Core/Private/GenericPlatform/GenericPlatformMisc.cpp
--- a/Engine/Source/Runtime/Core/Private/GenericPlatform/GenericPlatformMisc.cpp
+++ b/Engine/Source/Runtime/Core/Private/GenericPlatform/GenericPlatformMisc.cpp
@@ -505,6 +505,32 @@ const TCHAR* FGenericPlatformMisc::EngineDir()
 	return *EngineDirectory;
 }
 
+// wrap the LaunchDir variable in a function to work around static/global initialization order
+static FString& GetWrappedLaunchDir()
+{
+	static FString LaunchDir;
+	return LaunchDir;
+}
+
+void FGenericPlatformMisc::CacheLaunchDir()
+{
+	// we can only cache this ONCE
+	static bool bOneTime = false;
+	if (bOneTime)
+	{
+		return;
+	}
+	bOneTime = true;
+	
+	GetWrappedLaunchDir() = FPlatformProcess::GetCurrentWorkingDirectory() + TEXT("/");
+}
+
+const TCHAR* FGenericPlatformMisc::LaunchDir()
+{
+	return *GetWrappedLaunchDir();
+}
+
+
 const TCHAR* FGenericPlatformMisc::GetNullRHIShaderFormat()
 {
 	return TEXT("PCD3D_SM5");
diff --git a/Engine/Source/Runtime/Core/Private/GenericPlatform/GenericPlatformProcess.cpp b/Engine/Source/Runtime/Core/Private/GenericPlatform/GenericPlatformProcess.cpp
--- a/Engine/Source/Runtime/Core/Private/GenericPlatform/GenericPlatformProcess.cpp
+++ b/Engine/Source/Runtime/Core/Private/GenericPlatform/GenericPlatformProcess.cpp
@@ -83,6 +83,18 @@ const TCHAR* FGenericPlatformProcess::UserName(bool bOnlyAlphaNumeric/* = true*/
 	return TEXT("GenericUser");
 }
 
+void FGenericPlatformProcess::SetCurrentWorkingDirectoryToBaseDir()
+{
+	// even if we don't set a directory, we should remember the current one so LaunchDir works
+	FPlatformMisc::CacheLaunchDir();
+}
+
+FString FGenericPlatformProcess::GetCurrentWorkingDirectory()
+{
+	return TEXT("");
+}
+
+
 static FString Generic_ShaderDir;
 
 const TCHAR* FGenericPlatformProcess::ShaderDir()
diff --git a/Engine/Source/Runtime/Core/Private/Linux/LinuxPlatformProcess.cpp b/Engine/Source/Runtime/Core/Private/Linux/LinuxPlatformProcess.cpp
--- a/Engine/Source/Runtime/Core/Private/Linux/LinuxPlatformProcess.cpp
+++ b/Engine/Source/Runtime/Core/Private/Linux/LinuxPlatformProcess.cpp
@@ -921,7 +921,15 @@ uint32 FLinuxPlatformProcess::GetCurrentProcessId()
 	return getpid();
 }
 
-bool FLinuxPlatformProcess::GetProcReturnCode( FProcHandle& ProcHandle, int32* ReturnCode )
+FString FLinuxPlatformProcess::GetCurrentWorkingDirectory()
+{
+	// get the current directory
+	ANSICHAR CurrentDir[MAX_PATH] = { 0 };
+	getcwd(CurrentDir, sizeof(CurrentDir));
+	return UTF8_TO_TCHAR(CurrentDir);
+}
+
+bool FLinuxPlatformProcess::GetProcReturnCode(FProcHandle& ProcHandle, int32* ReturnCode)
 {
 	if (IsProcRunning(ProcHandle))
 	{
diff --git a/Engine/Source/Runtime/Core/Private/Mac/MacPlatformProcess.cpp b/Engine/Source/Runtime/Core/Private/Mac/MacPlatformProcess.cpp
--- a/Engine/Source/Runtime/Core/Private/Mac/MacPlatformProcess.cpp
+++ b/Engine/Source/Runtime/Core/Private/Mac/MacPlatformProcess.cpp
@@ -813,9 +813,18 @@ const TCHAR* FMacPlatformProcess::UserName(bool bOnlyAlphaNumeric)
 
 void FMacPlatformProcess::SetCurrentWorkingDirectoryToBaseDir()
 {
+	FPlatformMisc::CacheLaunchDir();
 	chdir([FString(BaseDir()).GetNSString() fileSystemRepresentation]);
 }
 
+FString FMacPlatformProcess::GetCurrentWorkingDirectory()
+{
+	// get the current directory
+	ANSICHAR CurrentDir[MAX_PATH] = { 0 };
+	getcwd(CurrentDir, sizeof(CurrentDir));
+	return UTF8_TO_TCHAR(CurrentDir);
+}
+
 const TCHAR* FMacPlatformProcess::ExecutableName(bool bRemoveExtension)
 {
 	static TCHAR Result[512]=TEXT("");
diff --git a/Engine/Source/Runtime/Core/Private/Misc/Paths.cpp b/Engine/Source/Runtime/Core/Private/Misc/Paths.cpp
--- a/Engine/Source/Runtime/Core/Private/Misc/Paths.cpp
+++ b/Engine/Source/Runtime/Core/Private/Misc/Paths.cpp
@@ -18,6 +18,11 @@ bool FPaths::ShouldSaveToUserDir()
 	return bShouldSaveToUserDir;
 }
 
+FString FPaths::LaunchDir()
+{
+	return FString(FPlatformMisc::LaunchDir());
+}
+
 FString FPaths::EngineDir()
 {
 	return FString(FPlatformMisc::EngineDir());
diff --git a/Engine/Source/Runtime/Core/Private/Windows/WindowsPlatformProcess.cpp b/Engine/Source/Runtime/Core/Private/Windows/WindowsPlatformProcess.cpp
--- a/Engine/Source/Runtime/Core/Private/Windows/WindowsPlatformProcess.cpp
+++ b/Engine/Source/Runtime/Core/Private/Windows/WindowsPlatformProcess.cpp
@@ -847,9 +847,19 @@ const TCHAR* FWindowsPlatformProcess::UserName(bool bOnlyAlphaNumeric/* = true*/
 
 void FWindowsPlatformProcess::SetCurrentWorkingDirectoryToBaseDir()
 {
+	FPlatformMisc::CacheLaunchDir();
 	verify(SetCurrentDirectoryW(BaseDir()));
 }
 
+/** Get the current working directory (only really makes sense on desktop platforms) */
+FString FWindowsPlatformProcess::GetCurrentWorkingDirectory()
+{
+	// get the current working directory (uncached)
+	TCHAR CurrentDirectory[MAX_PATH];
+	GetCurrentDirectoryW(MAX_PATH, CurrentDirectory);
+	return CurrentDirectory;
+}
+
 const TCHAR* FWindowsPlatformProcess::ExecutableName(bool bRemoveExtension)
 {
 	static TCHAR Result[512]=TEXT("");
diff --git a/Engine/Source/Runtime/Core/Public/GenericPlatform/GenericPlatformMisc.h b/Engine/Source/Runtime/Core/Public/GenericPlatform/GenericPlatformMisc.h
--- a/Engine/Source/Runtime/Core/Public/GenericPlatform/GenericPlatformMisc.h
+++ b/Engine/Source/Runtime/Core/Public/GenericPlatform/GenericPlatformMisc.h
@@ -603,6 +603,12 @@ public:
 	/** Get the engine directory */
 	static const TCHAR* EngineDir();
 
+	/** Get the directory the application was launched from (useful for commandline utilities) */
+	static const TCHAR* LaunchDir();
+
+	/** Function to store the current working directory for use with LaunchDir() */
+	static void CacheLaunchDir();
+
 	/**
 	 *	Return the GameDir
 	 */
diff --git a/Engine/Source/Runtime/Core/Public/GenericPlatform/GenericPlatformProcess.h b/Engine/Source/Runtime/Core/Public/GenericPlatform/GenericPlatformProcess.h
--- a/Engine/Source/Runtime/Core/Public/GenericPlatform/GenericPlatformProcess.h
+++ b/Engine/Source/Runtime/Core/Public/GenericPlatform/GenericPlatformProcess.h
@@ -231,7 +231,10 @@ struct CORE_API FGenericPlatformProcess
 	static const TCHAR* UserName(bool bOnlyAlphaNumeric = true);
 	static const TCHAR* ShaderDir();
 	static void SetShaderDir(const TCHAR*Where);
-	static void SetCurrentWorkingDirectoryToBaseDir() { }
+	static void SetCurrentWorkingDirectoryToBaseDir();
+	
+	/** Get the current working directory (only really makes sense on desktop platforms) */
+	static FString GetCurrentWorkingDirectory();
 
 	/**
 	 * Sets the process limits.
diff --git a/Engine/Source/Runtime/Core/Public/Linux/LinuxPlatformProcess.h b/Engine/Source/Runtime/Core/Public/Linux/LinuxPlatformProcess.h
--- a/Engine/Source/Runtime/Core/Public/Linux/LinuxPlatformProcess.h
+++ b/Engine/Source/Runtime/Core/Public/Linux/LinuxPlatformProcess.h
@@ -250,7 +250,8 @@ struct CORE_API FLinuxPlatformProcess : public FGenericPlatformProcess
 	static const TCHAR* UserDir();
 	static const TCHAR* UserSettingsDir();
 	static const TCHAR* ApplicationSettingsDir();
-	static FString GenerateApplicationPath( const FString& AppName, EBuildConfigurations::Type BuildConfiguration);
+	static FString GetCurrentWorkingDirectory();
+	static FString GenerateApplicationPath(const FString& AppName, EBuildConfigurations::Type BuildConfiguration);
 	static FString GetApplicationName( uint32 ProcessId );
 	static bool SetProcessLimits(EProcessResource::Type Resource, uint64 Limit);
 	static const TCHAR* ExecutableName(bool bRemoveExtension = true);
diff --git a/Engine/Source/Runtime/Core/Public/Mac/MacPlatformProcess.h b/Engine/Source/Runtime/Core/Public/Mac/MacPlatformProcess.h
--- a/Engine/Source/Runtime/Core/Public/Mac/MacPlatformProcess.h
+++ b/Engine/Source/Runtime/Core/Public/Mac/MacPlatformProcess.h
@@ -102,12 +102,13 @@ struct CORE_API FMacPlatformProcess : public FGenericPlatformProcess
 	static const TCHAR* ComputerName();
 	static const TCHAR* UserName(bool bOnlyAlphaNumeric = true);
 	static void SetCurrentWorkingDirectoryToBaseDir();
+	static FString GetCurrentWorkingDirectory();
 	static const TCHAR* ExecutableName(bool bRemoveExtension = true);
 	static FString GenerateApplicationPath( const FString& AppName, EBuildConfigurations::Type BuildConfiguration);
 	static const TCHAR* GetModuleExtension();
 	static const TCHAR* GetBinariesSubdirectory();
 	static const FString GetModulesDirectory();
-	static void LaunchURL( const TCHAR* URL, const TCHAR* Parms, FString* Error );
+	static void LaunchURL(const TCHAR* URL, const TCHAR* Parms, FString* Error);
 	static FProcHandle CreateProc( const TCHAR* URL, const TCHAR* Parms, bool bLaunchDetached, bool bLaunchHidden, bool bLaunchReallyHidden, uint32* OutProcessID, int32 PriorityModifier, const TCHAR* OptionalWorkingDirectory, void* PipeWrite );
 	static bool IsProcRunning( FProcHandle & ProcessHandle );
 	static void WaitForProc( FProcHandle & ProcessHandle );
diff --git a/Engine/Source/Runtime/Core/Public/Misc/Paths.h b/Engine/Source/Runtime/Core/Public/Misc/Paths.h
--- a/Engine/Source/Runtime/Core/Public/Misc/Paths.h
+++ b/Engine/Source/Runtime/Core/Public/Misc/Paths.h
@@ -15,6 +15,11 @@ public:
 	 */
 	static bool ShouldSaveToUserDir();
 
+	/**
+	 * Returns the directory the application was launched from (useful for commandline utilities)
+	 */
+	static FString LaunchDir();
+	 
 	/** 
 	 * Returns the base directory of the "core" engine that can be shared across
 	 * several games or across games & mods. Shaders and base localization files
@@ -480,7 +485,7 @@ private:
 
 	/** Holds the path to the currently loaded game project file. */
 	static FString GameProjectFilePath;
-	
+
 	/** Thread protection for above path */
 	FORCEINLINE static FCriticalSection* GameProjectFilePathLock() 
 	{
diff --git a/Engine/Source/Runtime/Core/Public/Windows/WindowsPlatformProcess.h b/Engine/Source/Runtime/Core/Public/Windows/WindowsPlatformProcess.h
--- a/Engine/Source/Runtime/Core/Public/Windows/WindowsPlatformProcess.h
+++ b/Engine/Source/Runtime/Core/Public/Windows/WindowsPlatformProcess.h
@@ -137,6 +137,7 @@ public:
 	static const TCHAR* ComputerName();
 	static const TCHAR* UserName(bool bOnlyAlphaNumeric = true);
 	static void SetCurrentWorkingDirectoryToBaseDir();
+	static FString GetCurrentWorkingDirectory();
 	static const TCHAR* ExecutableName(bool bRemoveExtension = true);
 	static FString GenerateApplicationPath( const FString& AppName, EBuildConfigurations::Type BuildConfiguration);
 	static const TCHAR* GetModuleExtension();
