- Fixed UnrealPak when creating pak files instead of opening existing ones

diff --git a/Engine/Source/Programs/UnrealPak/Private/UnrealPak.cpp b/Engine/Source/Programs/UnrealPak/Private/UnrealPak.cpp
--- a/Engine/Source/Programs/UnrealPak/Private/UnrealPak.cpp
+++ b/Engine/Source/Programs/UnrealPak/Private/UnrealPak.cpp
@@ -1085,6 +1085,26 @@ void RemoveIdenticalFiles( TArray<FPakInputPair>& FilesToPak, const FString& Sou
 	}
 }
 
+FString GetPakPath(const TCHAR* SpecifiedPath, bool bIsForCreation)
+{
+	FString PakFilename(SpecifiedPath);
+	FPaths::MakeStandardFilename(PakFilename);
+	
+	// if we are trying to open (not create) it, but BaseDir relative doesn't exist, look in LaunchDir
+	if (!bIsForCreation && !FPaths::FileExists(PakFilename))
+	{
+		PakFilename = FPaths::LaunchDir() + SpecifiedPath;
+
+		if (!FPaths::FileExists(PakFilename))
+		{
+			UE_LOG(LogPakFile, Fatal, TEXT("Existing pak file %s could not be found (checked against binary and launch directories)"), SpecifiedPath);
+			return TEXT("");
+		}
+	}
+	
+	return PakFilename;
+}
+
 /**
  * Application entry point
  * Params:
@@ -1141,31 +1161,19 @@ INT32_MAIN_INT32_ARGC_TCHAR_ARGV()
 	}
 	else 
 	{
-		FString PakFilename(ArgV[1]);
-		FPaths::MakeStandardFilename(PakFilename);
-		
-		// make sure it exists relative to engine, or relative to launch dir
-		if (!FPaths::FileExists(PakFilename))
-		{
-			PakFilename = FPaths::LaunchDir() + ArgV[1];
-		}
-			
-		if (!FPaths::FileExists(PakFilename))
-		{
-			UE_LOG(LogPakFile, Error, TEXT("Pak file %s could not be found (checked against binary and launch directories)"), ArgV[1]);
-			return 1;
-		}
-
 		if (FParse::Param(FCommandLine::Get(), TEXT("Test")))
 		{
+			FString PakFilename = GetPakPath(ArgV[1], false);
 			Result = TestPakFile(*PakFilename) ? 0 : 1;
 		}
 		else if (FParse::Param(FCommandLine::Get(), TEXT("List")))
 		{
+			FString PakFilename = GetPakPath(ArgV[1], false);
 			Result = ListFilesInPak(*PakFilename);
 		}
 		else if (FParse::Param(FCommandLine::Get(), TEXT("Extract")))
 		{
+			FString PakFilename = GetPakPath(ArgV[1], false);
 			if (ArgC < 4)
 			{
 				UE_LOG(LogPakFile, Error, TEXT("No extraction path specified."));
@@ -1179,6 +1187,9 @@ INT32_MAIN_INT32_ARGC_TCHAR_ARGV()
 		}
 		else
 		{
+			// since this is for creation, we pass true to make it not look in LaunchDir
+			FString PakFilename = GetPakPath(ArgV[1], true);
+
 			// List of all items to add to pak file
 			TArray<FPakInputPair> Entries;
 			ProcessCommandLine(ArgC, ArgV, Entries, CmdLineParameters);
