UnrealPak -List command now displays file entries in offset order

diff --git a/Engine/Source/Programs/UnrealPak/Private/UnrealPak.cpp b/Engine/Source/Programs/UnrealPak/Private/UnrealPak.cpp
--- a/Engine/Source/Programs/UnrealPak/Private/UnrealPak.cpp
+++ b/Engine/Source/Programs/UnrealPak/Private/UnrealPak.cpp
@@ -923,10 +923,27 @@ bool ListFilesInPak(const TCHAR * InPakFilename)
 
 	if (PakFile.IsValid())
 	{
-		for (FPakFile::FFileIterator It(PakFile); It; ++It, ++FileCount)
+		TArray<FPakFile::FFileIterator> Records;
+
+		for (FPakFile::FFileIterator It(PakFile); It; ++It)
+		{
+			Records.Add(It);
+		}
+
+		struct FOffsetSort
+		{
+			FORCEINLINE bool operator()(const FPakFile::FFileIterator& A, const FPakFile::FFileIterator& B) const
+			{
+				return A.Info().Offset < B.Info().Offset;
+			}
+		};
+
+		Records.Sort(FOffsetSort());
+
+		for (auto It : Records)
 		{
 			const FPakEntry& Entry = It.Info();
-			UE_LOG(LogPakFile, Display, TEXT("\"%s\" %d bytes."), *It.Filename(), Entry.Size);
+			UE_LOG(LogPakFile, Display, TEXT("\"%s\" offset: %lld, size: %d bytes."), *It.Filename(), Entry.Offset, Entry.Size);
 			FileSize += Entry.Size;
 			FileCount++;
 		}
