Merge 2678654.
PAK file changes for patch size reduction

diff --git a/Engine/Source/Programs/UnrealPak/Private/UnrealPak.cpp b/Engine/Source/Programs/UnrealPak/Private/UnrealPak.cpp
--- a/Engine/Source/Programs/UnrealPak/Private/UnrealPak.cpp
+++ b/Engine/Source/Programs/UnrealPak/Private/UnrealPak.cpp
@@ -17,11 +17,13 @@ struct FPakCommandLineParameters
 	FPakCommandLineParameters()
 		: CompressionBlockSize(64*1024)
 		, FileSystemBlockSize(0)
+		, PatchFilePadAlign(0)
 		, GeneratePatch(false)
 	{}
 
 	int32  CompressionBlockSize;
 	int64  FileSystemBlockSize;
+	int64  PatchFilePadAlign;
 	bool   GeneratePatch;
 	FString SourcePatchPakFilename;
 	FString SourcePatchDiffDirectory;
@@ -410,6 +412,11 @@ void ProcessCommandLine(int32 ArgC, TCHAR* ArgV[], TArray<FPakInputPair>& Entrie
 		CmdLineParameters.FileSystemBlockSize = 0;
 	}
 
+	if (!FParse::Value(FCommandLine::Get(), TEXT("-patchpaddingalign="), CmdLineParameters.PatchFilePadAlign))
+	{
+		CmdLineParameters.PatchFilePadAlign = 0;
+	}
+
 	if (FParse::Value(FCommandLine::Get(), TEXT("-create="), ResponseFile))
 	{
 		TArray<FString> Lines;
@@ -750,6 +757,13 @@ bool CreatePakFile(const TCHAR* Filename, TArray<FPakInputPair>& FilesToAdd, con
 	ECompressionFlags CompressionMethod = COMPRESS_None;
 	FCompressedFileBuffer CompressedFileBuffer;
 
+	uint8* PaddingBuffer = nullptr;
+	if (CmdLineParameters.PatchFilePadAlign > 0)
+	{
+		PaddingBuffer = (uint8*)FMemory::Malloc(CmdLineParameters.PatchFilePadAlign);
+		FMemory::Memset(PaddingBuffer, 0, CmdLineParameters.PatchFilePadAlign);
+	}
+
 	for (int32 FileIndex = 0; FileIndex < FilesToAdd.Num(); FileIndex++)
 	{
 		//  Remember the offset but don't serialize it with the entry header.
@@ -794,6 +808,22 @@ bool CreatePakFile(const TCHAR* Filename, TArray<FPakInputPair>& FilesToAdd, con
 			}
 		}
 
+		// Some platforms provide patch download size reduction by diffing the patch files.  However, they often operate on specific block
+		// sizes when dealing with new data within the file.  Pad files out to the given alignment to work with these systems more nicely.
+		if (CmdLineParameters.PatchFilePadAlign > 0 && OriginalFileSize != INDEX_NONE)
+		{	
+			NewEntryOffset = AlignArbitrary(NewEntryOffset, CmdLineParameters.PatchFilePadAlign);
+			int64 CurrentLoc = PakFileHandle->Tell();
+			int64 PaddingSize = NewEntryOffset - CurrentLoc;
+			check(PaddingSize <= CmdLineParameters.PatchFilePadAlign);
+
+			//have to pad manually with 0's.  File locations skipped by Seek and never written are uninitialized which would defeat the whole purpose
+			//of padding for certain platforms patch diffing systems.
+			PakFileHandle->Serialize(PaddingBuffer, PaddingSize);
+			check(PakFileHandle->Tell() == NewEntryOffset);						
+		}
+
+		
 		bool bCopiedToPak;
 		if (FilesToAdd[FileIndex].bNeedsCompression && CompressionMethod != COMPRESS_None)
 		{
@@ -803,7 +833,7 @@ bool CreatePakFile(const TCHAR* Filename, TArray<FPakInputPair>& FilesToAdd, con
 		{
 			bCopiedToPak = CopyFileToPak(*PakFileHandle, MountPoint, FilesToAdd[FileIndex], ReadBuffer, BufferSize, NewEntry);
 		}
-
+		
 		if (bCopiedToPak)
 		{
 			// Update offset now and store it in the index (and only in index)
@@ -825,6 +855,7 @@ bool CreatePakFile(const TCHAR* Filename, TArray<FPakInputPair>& FilesToAdd, con
 		}
 	}
 
+	FMemory::Free(PaddingBuffer);
 	FMemory::Free(ReadBuffer);
 	ReadBuffer = NULL;
 
