Prevent logging an error message to say "including 0 errors". Causes build failure for every run of DiffedCookBuild.

diff --git a/Engine/Source/Programs/UnrealPak/Private/UnrealPak.cpp b/Engine/Source/Programs/UnrealPak/Private/UnrealPak.cpp
--- a/Engine/Source/Programs/UnrealPak/Private/UnrealPak.cpp
+++ b/Engine/Source/Programs/UnrealPak/Private/UnrealPak.cpp
@@ -1049,7 +1049,7 @@ bool ExtractFilesFromPak(const TCHAR* InPakFilename, const TCHAR* InDestPath, bo
 		FMemory::Free(Buffer);
 		FMemory::Free(PersistantCompressionBuffer);
 
-		UE_LOG(LogPakFile, Error, TEXT("Finished extracting %d files (including %d errors)."), FileCount, ErrorCount);
+		UE_LOG(LogPakFile, Log, TEXT("Finished extracting %d files (including %d errors)."), FileCount, ErrorCount);
 
 		return true;
 	}
