Rollback //UE4/Main/Engine/Source/Programs/UnrealPak/Private/UnrealPak.cpp to revision 11

Does not work with compressed pak files.

diff --git a/Engine/Source/Programs/UnrealPak/Private/UnrealPak.cpp b/Engine/Source/Programs/UnrealPak/Private/UnrealPak.cpp
--- a/Engine/Source/Programs/UnrealPak/Private/UnrealPak.cpp
+++ b/Engine/Source/Programs/UnrealPak/Private/UnrealPak.cpp
@@ -234,7 +234,7 @@ bool FCompressedFileBuffer::CompressFileToWorkingBuffer(const FPakInputPair& InF
 	return true;
 }
 
-bool PrepareCopyFileToPak(FArchive& InPak, const FString& InMountPoint, const FPakInputPair& InFile, uint8*& InOutPersistentBuffer, int64& InOutBufferSize, FPakEntryPair& OutNewEntry, uint8*& OutDataToWrite, int64& OutSizeToWrite)
+bool CopyFileToPak(FArchive& InPak, const FString& InMountPoint, const FPakInputPair& InFile, uint8*& InOutPersistentBuffer, int64& InOutBufferSize, FPakEntryPair& OutNewEntry)
 {	
 	TAutoPtr<FArchive> FileHandle(IFileManager::Get().CreateFileReader(*InFile.Source));
 	bool bFileExists = FileHandle.IsValid();
@@ -259,7 +259,7 @@ bool PrepareCopyFileToPak(FArchive& InPak, const FString& InMountPoint, const FP
 		FileHandle->Serialize(InOutPersistentBuffer, FileSize);
 
 		{
-			OutSizeToWrite = FileSize;
+			int64 SizeToWrite = FileSize;
 			if (InFile.bNeedEncryption)
 			{
 				for(int64 FillIndex=FileSize; FillIndex < PaddedEncryptedFileSize && InFile.bNeedEncryption; ++FillIndex)
@@ -271,19 +271,22 @@ bool PrepareCopyFileToPak(FArchive& InPak, const FString& InMountPoint, const FP
 				//Encrypt the buffer before writing it to disk
 				FAES::EncryptData(InOutPersistentBuffer, PaddedEncryptedFileSize);
 				// Update the size to be written
-				OutSizeToWrite = PaddedEncryptedFileSize;
+				SizeToWrite = PaddedEncryptedFileSize;
 				OutNewEntry.Info.bEncrypted = true;
 			}
 
 			// Calculate the buffer hash value
-			FSHA1::HashBuffer(InOutPersistentBuffer,FileSize,OutNewEntry.Info.Hash);			
-			OutDataToWrite = InOutPersistentBuffer;
+			FSHA1::HashBuffer(InOutPersistentBuffer,FileSize,OutNewEntry.Info.Hash);
+
+			// Write to file
+			OutNewEntry.Info.Serialize(InPak,FPakInfo::PakFile_Version_Latest);
+			InPak.Serialize(InOutPersistentBuffer,SizeToWrite);
 		}
 	}
 	return bFileExists;
 }
 
-bool PrepareCopyCompressedFileToPak(FArchive& InPak, const FString& InMountPoint, const FPakInputPair& InFile, const FCompressedFileBuffer& CompressedFile, FPakEntryPair& OutNewEntry, uint8*& OutDataToWrite, int64& OutSizeToWrite)
+bool CopyCompressedFileToPak(FArchive& InPak, const FString& InMountPoint, const FPakInputPair& InFile, const FCompressedFileBuffer& CompressedFile, FPakEntryPair& OutNewEntry)
 {
 	if (CompressedFile.TotalCompressedSize == 0)
 	{
@@ -322,10 +325,8 @@ bool PrepareCopyCompressedFileToPak(FArchive& InPak, const FString& InMountPoint
 	OutNewEntry.Filename = InFile.Dest.Mid(InMountPoint.Len());
 	OutNewEntry.Info.Offset = 0; // Don't serialize offsets here.
 	OutNewEntry.Info.bEncrypted = InFile.bNeedEncryption;
-	OutSizeToWrite = CompressedFile.TotalCompressedSize;
-	OutDataToWrite = CompressedFile.CompressedBuffer.Get();
-	//OutNewEntry.Info.Serialize(InPak,FPakInfo::PakFile_Version_Latest);	
-	//InPak.Serialize(CompressedFile.CompressedBuffer.Get(), CompressedFile.TotalCompressedSize);
+	OutNewEntry.Info.Serialize(InPak,FPakInfo::PakFile_Version_Latest);
+	InPak.Serialize(CompressedFile.CompressedBuffer.Get(), CompressedFile.TotalCompressedSize);
 
 	return true;
 }
@@ -769,15 +770,6 @@ bool CreatePakFile(const TCHAR* Filename, TArray<FPakInputPair>& FilesToAdd, con
 		FMemory::Memset(PaddingBuffer, 0, PaddingBufferSize);
 	}
 
-	// Some platforms provide patch download size reduction by diffing the patch files.  However, they often operate on specific block
-	// sizes when dealing with new data within the file.  Pad files out to the given alignment to work with these systems more nicely.
-	// We also want to combine smaller files into the same padding size block so we don't waste as much space. i.e. grouping 64 1k files together
-	// rather than padding each out to 64k.
-	const uint32 RequiredPatchPadding = CmdLineParameters.PatchFilePadAlign;
-
-	uint64 ContiguousTotalSizeSmallerThanBlockSize = 0;
-	uint64 ContiguousFilesSmallerThanBlockSize = 0;
-
 	for (int32 FileIndex = 0; FileIndex < FilesToAdd.Num(); FileIndex++)
 	{
 		//  Remember the offset but don't serialize it with the entry header.
@@ -796,8 +788,8 @@ bool CreatePakFile(const TCHAR* Filename, TArray<FPakInputPair>& FilesToAdd, con
 				// Check the compression ratio, if it's too low just store uncompressed. Also take into account read size
 				// if we still save 64KB it's probably worthwhile compressing, as that saves a file read operation in the runtime.
 				// TODO: drive this threashold from the command line
-				float PercentLess = ((float)CompressedFileBuffer.TotalCompressedSize / (OriginalFileSize / 100.f));
-				if (PercentLess > 90.f && (OriginalFileSize - CompressedFileBuffer.TotalCompressedSize) < 65536)
+				float PercentLess = ((float)CompressedFileBuffer.TotalCompressedSize/(OriginalFileSize/100.f));
+				if (PercentLess > 90.f && (OriginalFileSize-CompressedFileBuffer.TotalCompressedSize) < 65536)
 				{
 					CompressionMethod = COMPRESS_None;
 				}
@@ -814,13 +806,13 @@ bool CreatePakFile(const TCHAR* Filename, TArray<FPakInputPair>& FilesToAdd, con
 		// Account for file system block size, which is a boundary we want to avoid crossing.
 		if (CmdLineParameters.FileSystemBlockSize > 0 && OriginalFileSize != INDEX_NONE && RealFileSize <= CmdLineParameters.FileSystemBlockSize)
 		{
-			if ((NewEntryOffset / CmdLineParameters.FileSystemBlockSize) != ((NewEntryOffset + RealFileSize) / CmdLineParameters.FileSystemBlockSize))
+			if ((NewEntryOffset / CmdLineParameters.FileSystemBlockSize) != ((NewEntryOffset+RealFileSize) / CmdLineParameters.FileSystemBlockSize))
 			{
 				//File crosses a block boundary, so align it to the beginning of the next boundary
 				int64 OldOffset = NewEntryOffset;
 				NewEntryOffset = AlignArbitrary(NewEntryOffset, CmdLineParameters.FileSystemBlockSize);
 				int64 PaddingRequired = NewEntryOffset - OldOffset;
-
+				
 				if (PaddingRequired > 0)
 				{
 					// If we don't already have a padding buffer, create one
@@ -837,78 +829,40 @@ bool CreatePakFile(const TCHAR* Filename, TArray<FPakInputPair>& FilesToAdd, con
 						PakFileHandle->Serialize(PaddingBuffer, AmountToWrite);
 						PaddingRequired -= AmountToWrite;
 					}
-
+					
 					check(PakFileHandle->Tell() == NewEntryOffset);
 				}
 			}
 		}
 
+		// Some platforms provide patch download size reduction by diffing the patch files.  However, they often operate on specific block
+		// sizes when dealing with new data within the file.  Pad files out to the given alignment to work with these systems more nicely.
+		if (CmdLineParameters.PatchFilePadAlign > 0 && OriginalFileSize != INDEX_NONE)
+		{	
+			NewEntryOffset = AlignArbitrary(NewEntryOffset, CmdLineParameters.PatchFilePadAlign);
+			int64 CurrentLoc = PakFileHandle->Tell();
+			int64 PaddingSize = NewEntryOffset - CurrentLoc;
+			check(PaddingSize <= PaddingBufferSize);
+
+			//have to pad manually with 0's.  File locations skipped by Seek and never written are uninitialized which would defeat the whole purpose
+			//of padding for certain platforms patch diffing systems.
+			PakFileHandle->Serialize(PaddingBuffer, PaddingSize);
+			check(PakFileHandle->Tell() == NewEntryOffset);						
+		}
+
+		
 		bool bCopiedToPak;
-		int64 SizeToWrite = 0;
-		uint8* DataToWrite = nullptr;
 		if (FilesToAdd[FileIndex].bNeedsCompression && CompressionMethod != COMPRESS_None)
 		{
-			bCopiedToPak = PrepareCopyCompressedFileToPak(*PakFileHandle, MountPoint, FilesToAdd[FileIndex], CompressedFileBuffer, NewEntry, DataToWrite, SizeToWrite);
-			DataToWrite = CompressedFileBuffer.CompressedBuffer.Get();
+			bCopiedToPak = CopyCompressedFileToPak(*PakFileHandle, MountPoint, FilesToAdd[FileIndex], CompressedFileBuffer, NewEntry);
 		}
 		else
 		{
-			bCopiedToPak = PrepareCopyFileToPak(*PakFileHandle, MountPoint, FilesToAdd[FileIndex], ReadBuffer, BufferSize, NewEntry, DataToWrite, SizeToWrite);
-			DataToWrite = ReadBuffer;
-		}		
-
-		int64 TotalSizeToWrite = SizeToWrite + NewEntry.Info.GetSerializedSize(FPakInfo::PakFile_Version_Latest);
+			bCopiedToPak = CopyFileToPak(*PakFileHandle, MountPoint, FilesToAdd[FileIndex], ReadBuffer, BufferSize, NewEntry);
+		}
+		
 		if (bCopiedToPak)
 		{
-			
-			if (RequiredPatchPadding > 0)
-			{
-				//if the next file is going to cross a patch-block boundary then pad out the current set of files with 0's
-				//and align the next file up.
-				bool bCrossesBoundary = AlignArbitrary(NewEntryOffset, RequiredPatchPadding) != AlignArbitrary(NewEntryOffset + TotalSizeToWrite - 1, RequiredPatchPadding);
-				if (TotalSizeToWrite >= RequiredPatchPadding || // if it exactly the padding size and by luck does not cross a boundary, we still consider it "over" because it can't be packed with anything else
-					bCrossesBoundary)
-				{
-					NewEntryOffset = AlignArbitrary(NewEntryOffset, RequiredPatchPadding);
-					int64 CurrentLoc = PakFileHandle->Tell();
-					int64 PaddingSize = NewEntryOffset - CurrentLoc;
-					check(PaddingSize >= 0);
-					if (PaddingSize)
-					{
-						check(PaddingSize <= PaddingBufferSize);
-
-						//have to pad manually with 0's.  File locations skipped by Seek and never written are uninitialized which would defeat the whole purpose
-						//of padding for certain platforms patch diffing systems.
-						PakFileHandle->Serialize(PaddingBuffer, PaddingSize);
-					}
-					check(PakFileHandle->Tell() == NewEntryOffset);
-				}
-
-				//if the current file is bigger than a patch block then we will always have to pad out the previous files.
-				//if there were a large set of contiguous small files behind us then this will be the natural stopping point for a possible pathalogical patching case where growth in the small files causes a cascade 
-				//to dirty up all the blocks prior to this one.  If this could happen let's warn about it.
-				if (TotalSizeToWrite >= RequiredPatchPadding || 
-					FileIndex + 1 == FilesToAdd.Num()) // also check the last file, this won't work perfectly if we don't end up adding the last file for some reason
-				{
-					const uint64 ContiguousGroupedFilePatchWarningThreshhold = 50 * 1024 * 1024;
-					if (ContiguousTotalSizeSmallerThanBlockSize > ContiguousGroupedFilePatchWarningThreshhold)
-					{
-						UE_LOG(LogPakFile, Warning, TEXT("%i small files (%i) totaling %llu contiguous bytes found before first 'large' file.  Changes to any of these files could cause the whole group to be 'dirty' in a per-file binary diff based patching system."), ContiguousFilesSmallerThanBlockSize, RequiredPatchPadding, ContiguousTotalSizeSmallerThanBlockSize);
-					}
-					ContiguousTotalSizeSmallerThanBlockSize = 0;
-					ContiguousFilesSmallerThanBlockSize = 0;
-				}
-				else
-				{
-					ContiguousTotalSizeSmallerThanBlockSize += TotalSizeToWrite;
-					ContiguousFilesSmallerThanBlockSize++;				
-				}
-			}
-
-			// Write to file
-			NewEntry.Info.Serialize(*PakFileHandle, FPakInfo::PakFile_Version_Latest);
-			PakFileHandle->Serialize(DataToWrite, SizeToWrite);	
-
 			// Update offset now and store it in the index (and only in index)
 			NewEntry.Info.Offset = NewEntryOffset;
 			Index.Add(NewEntry);
@@ -947,20 +901,6 @@ bool CreatePakFile(const TCHAR* Filename, TArray<FPakInputPair>& FilesToAdd, con
 		FPakEntryPair& Entry = Index[EntryIndex];
 		IndexWriter << Entry.Filename;
 		Entry.Info.Serialize(IndexWriter, Info.Version);
-
-		if (RequiredPatchPadding > 0)
-		{
-			int64 EntrySize = Entry.Info.GetSerializedSize(FPakInfo::PakFile_Version_Latest);
-			int64 TotalSizeToWrite = Entry.Info.Size + EntrySize;
-			if (TotalSizeToWrite >= RequiredPatchPadding)
-			{
-				int64 RealStart = Entry.Info.Offset;
-				if ((RealStart % RequiredPatchPadding) != 0)
-				{
-					UE_LOG(LogPakFile, Warning, TEXT("File at offset %lld of size %lld not aligned to patch size %i"), RealStart, Entry.Info.Size, RequiredPatchPadding);
-				}
-			}
-		}
 	}
 	PakFileHandle->Serialize(IndexData.GetData(), IndexData.Num());
 
