Copying //UE4/Dev-Core to //UE4/Dev-Main (Source: //UE4/Dev-Core @ 3208226)

diff --git a/Engine/Source/Programs/CrashReporter/MinidumpDiagnostics/MinidumpDiagnostics.Target.cs b/Engine/Source/Programs/CrashReporter/MinidumpDiagnostics/MinidumpDiagnostics.Target.cs
--- a/Engine/Source/Programs/CrashReporter/MinidumpDiagnostics/MinidumpDiagnostics.Target.cs
+++ b/Engine/Source/Programs/CrashReporter/MinidumpDiagnostics/MinidumpDiagnostics.Target.cs
@@ -61,5 +61,6 @@ public class MinidumpDiagnosticsTarget : TargetRules
 		OutLinkEnvironmentConfiguration.bIsBuildingConsoleApplication = true;
 
 		OutCPPEnvironmentConfiguration.Definitions.Add("MINIDUMPDIAGNOSTICS=1");
-	}
+        OutCPPEnvironmentConfiguration.Definitions.Add("NOINITCRASHREPORTER=1");
+    }
 }
diff --git a/Engine/Source/Programs/UnrealBuildTool/Configuration/EngineConfiguration.cs b/Engine/Source/Programs/UnrealBuildTool/Configuration/EngineConfiguration.cs
--- a/Engine/Source/Programs/UnrealBuildTool/Configuration/EngineConfiguration.cs
+++ b/Engine/Source/Programs/UnrealBuildTool/Configuration/EngineConfiguration.cs
@@ -39,7 +39,7 @@ namespace UnrealBuildTool
 																	"/Script/BuildSettings.BuildSettings", "/Script/IOSRuntimeSettings.IOSRuntimeSettings", "/Script/WindowsTargetPlatform.WindowsTargetSettings",
 																	"/Script/UnrealEd.ProjectPackagingSettings", "/Script/PS4PlatformEditor.PS4TargetSettings", "/Script/XboxOnePlatformEditor.XboxOneTargetSettings",
 																	"/Script/HTML5PlatformEditor.HTML5TargetSettings","PS4SymbolServer","/Script/EngineSettings.GeneralProjectSettings",
-                                                                    "/Script/UnrealEd.ProjectPackagingSettings", "InstalledPlatforms", "OnlineSubsystemGooglePlay.Store","/Script/Engine.StreamingSettings",
+                                                                    "/Script/UnrealEd.ProjectPackagingSettings", "InstalledPlatforms", "OnlineSubsystemGooglePlay.Store","/Script/Engine.StreamingSettings", "Core.Encryption",
                                                                     "/Script/SwitchRuntimeSettings.SwitchRuntimeSettings", "/Script/Engine.NetworkSettings", "/Script/SourceCodeAccess.SourceCodeAccessSettings",
 																	"/Script/LinuxTargetPlatform.LinuxTargetSettings"};
 
@@ -854,7 +854,7 @@ namespace UnrealBuildTool
 		/// <summary>
 		/// Returns the platform name to use as part of platform-specific config files
 		/// </summary>
-		private static string GetIniPlatformName(UnrealTargetPlatform TargetPlatform)
+		public static string GetIniPlatformName(UnrealTargetPlatform TargetPlatform)
 		{
 			if (TargetPlatform == UnrealTargetPlatform.Win32 || TargetPlatform == UnrealTargetPlatform.Win64)
 			{
diff --git a/Engine/Source/Programs/UnrealBuildTool/Configuration/UEBuildTarget.cs b/Engine/Source/Programs/UnrealBuildTool/Configuration/UEBuildTarget.cs
--- a/Engine/Source/Programs/UnrealBuildTool/Configuration/UEBuildTarget.cs
+++ b/Engine/Source/Programs/UnrealBuildTool/Configuration/UEBuildTarget.cs
@@ -3881,6 +3881,38 @@ namespace UnrealBuildTool
 			// Default does nothing
 		}
 
+		public static void ParseEncryptionIni(DirectoryReference InDirectory, UnrealTargetPlatform InTargetPlatform, out String[] OutRSAKeys, out String OutAESKey)
+		{
+			OutAESKey = String.Empty;
+			OutRSAKeys = null;
+
+			ConfigCacheIni Ini = ConfigCacheIni.CreateConfigCacheIni(InTargetPlatform, "Encryption", InDirectory);
+
+			bool bSigningEnabled;
+			Ini.GetBool("Core.Encryption", "SignPak", out bSigningEnabled);
+
+			if (bSigningEnabled)
+			{
+				OutRSAKeys = new string[3];
+				Ini.GetString("Core.Encryption", "rsa.privateexp", out OutRSAKeys[0]);
+				Ini.GetString("Core.Encryption", "rsa.modulus", out OutRSAKeys[1]);
+				Ini.GetString("Core.Encryption", "rsa.publicexp", out OutRSAKeys[2]);
+
+				if (String.IsNullOrEmpty(OutRSAKeys[0]) || String.IsNullOrEmpty(OutRSAKeys[1]) || String.IsNullOrEmpty(OutRSAKeys[2]))
+				{
+					OutRSAKeys = null;
+				}
+			}
+
+			bool bEncryptionEnabled;
+			Ini.GetBool("Core.Encryption", "EncryptPak", out bEncryptionEnabled);
+
+			if (bEncryptionEnabled)
+			{
+				Ini.GetString("Core.Encryption", "aes.key", out OutAESKey);
+			}
+		}
+
 		/// <summary>
 		/// Sets up the global compile and link environment for the target.
 		/// </summary>
@@ -3928,8 +3960,33 @@ namespace UnrealBuildTool
 				GlobalCompileEnvironment.Config.Definitions.Add(String.Format("IS_PROGRAM={0}", TargetType == TargetRules.TargetType.Program ? "1" : "0"));
 			}
 
-			// See if the target rules defined any signing keys and add the appropriate defines
-			if (!string.IsNullOrEmpty(Rules.PakSigningKeysFile))
+			String[] RSAKeys;
+			String AESKey;
+			ParseEncryptionIni(ProjectDirectory, Platform, out RSAKeys, out AESKey);
+			bool bSigningEnabledByIniFile = false;
+
+			if (RSAKeys != null)
+			{
+				GlobalCompileEnvironment.Config.Definitions.Add("DECRYPTION_KEY_MODULUS=\"" + RSAKeys[1] + "\"");
+				GlobalCompileEnvironment.Config.Definitions.Add("DECRYPTION_KEY_EXPONENT=\"" + RSAKeys[2] + "\"");
+
+				bSigningEnabledByIniFile = true;
+			}
+
+			if (!String.IsNullOrEmpty(AESKey))
+			{
+				if (AESKey.Length < 32)
+				{
+					Log.TraceError("AES key specified in configs must be at least 32 characters long!");
+				}
+				else
+				{
+					GlobalCompileEnvironment.Config.Definitions.Add("AES_KEY=\"" + AESKey + "\"");
+				}
+			}
+			
+			// If we didn't extract any keys from the new ini file setup, try looking for the old keys text file
+			if (!bSigningEnabledByIniFile && !string.IsNullOrEmpty(Rules.PakSigningKeysFile))
 			{
 				string FullFilename = Path.Combine(ProjectDirectory.FullName, Rules.PakSigningKeysFile);
 
diff --git a/Engine/Source/Programs/UnrealBuildTool/System/UnrealBuildTool.cs b/Engine/Source/Programs/UnrealBuildTool/System/UnrealBuildTool.cs
--- a/Engine/Source/Programs/UnrealBuildTool/System/UnrealBuildTool.cs
+++ b/Engine/Source/Programs/UnrealBuildTool/System/UnrealBuildTool.cs
@@ -1748,14 +1748,24 @@ namespace UnrealBuildTool
 							foreach (TargetDescriptor Desc in TargetDescs)
 							{
 								DirectoryReference ProjectDirectory = DirectoryReference.FromFile(ProjectFile);
-								foreach (FileReference IniFilename in ConfigCacheIni.EnumerateCrossPlatformIniFileNames(ProjectDirectory, UnrealBuildTool.EngineDirectory, Desc.Platform, "Engine", false))
+								String[] IniBaseNames = new String[] { "Engine", "Encryption" };
+
+								foreach (String BaseName in IniBaseNames)
 								{
-									FileInfo IniFileInfo = new FileInfo(IniFilename.FullName);
-									if (UBTMakefileInfo.LastWriteTime.CompareTo(IniFileInfo.LastWriteTime) < 0)
+									foreach (FileReference IniFilename in ConfigCacheIni.EnumerateCrossPlatformIniFileNames(ProjectDirectory, UnrealBuildTool.EngineDirectory, Desc.Platform, BaseName, false))
+									{
+										FileInfo IniFileInfo = new FileInfo(IniFilename.FullName);
+										if (UBTMakefileInfo.LastWriteTime.CompareTo(IniFileInfo.LastWriteTime) < 0)
+										{
+											// Ini files are newer than UBTMakefile
+											UBTMakefile = null;
+											ReasonNotLoaded = "ini files are newer that UBTMakefile";
+											break;
+										}
+									}
+
+									if (UBTMakefile == null)
 									{
-										// Ini files are newer than UBTMakefile
-										UBTMakefile = null;
-										ReasonNotLoaded = "ini files are newer that UBTMakefile";
 										break;
 									}
 								}
diff --git a/Engine/Source/Programs/UnrealFrontend/Private/Commands/StatsConvertCommand.cpp b/Engine/Source/Programs/UnrealFrontend/Private/Commands/StatsConvertCommand.cpp
--- a/Engine/Source/Programs/UnrealFrontend/Private/Commands/StatsConvertCommand.cpp
+++ b/Engine/Source/Programs/UnrealFrontend/Private/Commands/StatsConvertCommand.cpp
@@ -2,6 +2,7 @@
 
 #include "UnrealFrontendPrivatePCH.h"
 #include "StatsConvertCommand.h"
+#include "UniquePtr.h"
 
 /** Helper class used to extract stats data into CSV file. */
 class FCSVStatsProfiler : public FStatsReadFile
@@ -129,14 +130,14 @@ void FStatsConvertCommand::InternalRun()
 	FParse::Value(FCommandLine::Get(), TEXT("-STATLIST="), StatListString);
 
 	// Open a CSV file for write.
-	TAutoPtr<FArchive> FileWriter( IFileManager::Get().CreateFileWriter( *OutFile ) );
+	TUniquePtr<FArchive> FileWriter( IFileManager::Get().CreateFileWriter( *OutFile ) );
 	if (!FileWriter)
 	{
 		UE_LOG( LogStats, Error, TEXT( "Could not open output file: %s" ), *OutFile );
 		return;
 	}
 
-	TAutoPtr<FCSVStatsProfiler> Instance( FStatsReader<FCSVStatsProfiler>::Create( *TargetFile ) );
+	TUniquePtr<FCSVStatsProfiler> Instance( FStatsReader<FCSVStatsProfiler>::Create( *TargetFile ) );
 	if (Instance)
 	{
 		TArray<FName> StatList;
@@ -148,7 +149,7 @@ void FStatsConvertCommand::InternalRun()
 			StatArrayString.Add( TEXT( "STAT_FrameTime" ) );
 		}
 
-		Instance->Initialize( FileWriter, StatArrayString );
+		Instance->Initialize( FileWriter.Get(), StatArrayString );
 		Instance->ReadAndProcessAsynchronously();
 
 		while (Instance->IsBusy())
diff --git a/Engine/Source/Programs/UnrealHeaderTool/Private/BaseParser.cpp b/Engine/Source/Programs/UnrealHeaderTool/Private/BaseParser.cpp
--- a/Engine/Source/Programs/UnrealHeaderTool/Private/BaseParser.cpp
+++ b/Engine/Source/Programs/UnrealHeaderTool/Private/BaseParser.cpp
@@ -449,11 +449,11 @@ bool FBaseParser::GetToken( FToken& Token, bool bNoConsts/*=false*/, ESymbolPars
 		else if (bIsHex)
 		{
 			TCHAR* End = Token.Identifier + FCString::Strlen(Token.Identifier);
-			Token.SetConstInt( FCString::Strtoi(Token.Identifier,&End,0) );
+			Token.SetConstInt64( FCString::Strtoi64(Token.Identifier,&End,0) );
 		}
 		else
 		{
-			Token.SetConstInt( FCString::Atoi(Token.Identifier) );
+			Token.SetConstInt64( FCString::Atoi64(Token.Identifier) );
 		}
 		return true;
 	}
@@ -735,6 +735,29 @@ bool FBaseParser::GetConstInt(int32& Result, const TCHAR* Tag)
 	return false;
 }
 
+bool FBaseParser::GetConstInt64(int64& Result, const TCHAR* Tag)
+{
+	FToken Token;
+	if (GetToken(Token))
+	{
+		if (Token.GetConstInt64(Result))
+		{
+			return true;
+		}
+		else
+		{
+			UngetToken(Token);
+		}
+	}
+
+	if (Tag != NULL)
+	{
+		FError::Throwf(TEXT("%s: Missing constant integer"), Tag );
+	}
+
+	return false;
+}
+
 bool FBaseParser::MatchSymbol( const TCHAR* Match, ESymbolParseOption bParseTemplateClosingBracket/*=ESymbolParseOption::Normal*/ )
 {
 	FToken Token;
@@ -798,7 +821,7 @@ bool FBaseParser::MatchConstInt( const TCHAR* Match )
 	FToken Token;
 	if (GetToken(Token))
 	{
-		if( Token.TokenType==TOKEN_Const && Token.Type == CPT_Int && FCString::Stricmp(Token.Identifier,Match)==0 )
+		if( Token.TokenType==TOKEN_Const && (Token.Type == CPT_Int || Token.Type == CPT_Int64) && FCString::Stricmp(Token.Identifier,Match)==0 )
 		{
 			return true;
 		}
diff --git a/Engine/Source/Programs/UnrealHeaderTool/Private/BaseParser.h b/Engine/Source/Programs/UnrealHeaderTool/Private/BaseParser.h
--- a/Engine/Source/Programs/UnrealHeaderTool/Private/BaseParser.h
+++ b/Engine/Source/Programs/UnrealHeaderTool/Private/BaseParser.h
@@ -132,6 +132,7 @@ protected:
 	 * @return true on success, otherwise false.
 	 */
 	bool GetConstInt(int32& Result, const TCHAR* Tag = NULL);
+	bool GetConstInt64(int64& Result, const TCHAR* Tag = NULL);
 
 	// Matching predefined text.
 	bool MatchIdentifier( FName Match );
diff --git a/Engine/Source/Programs/UnrealHeaderTool/Private/ClassMaps.cpp b/Engine/Source/Programs/UnrealHeaderTool/Private/ClassMaps.cpp
--- a/Engine/Source/Programs/UnrealHeaderTool/Private/ClassMaps.cpp
+++ b/Engine/Source/Programs/UnrealHeaderTool/Private/ClassMaps.cpp
@@ -17,7 +17,7 @@ TSet<FUnrealSourceFile*> GExportedSourceFiles;
 TMap<UProperty*, FString> GArrayDimensions;
 TMap<UPackage*,  const FManifestModule*> GPackageToManifestModuleMap;
 TMap<UField*, uint32> GGeneratedCodeCRCs;
-TMap<UEnum*,  EPropertyType> GEnumUnderlyingTypes;
+TMap<UEnum*,  EUnderlyingEnumType> GEnumUnderlyingTypes;
 TMap<FName, TSharedRef<FClassDeclarationMetaData> > GClassDeclarations;
 TSet<UProperty*> GUnsizedProperties;
 
diff --git a/Engine/Source/Programs/UnrealHeaderTool/Private/ClassMaps.h b/Engine/Source/Programs/UnrealHeaderTool/Private/ClassMaps.h
--- a/Engine/Source/Programs/UnrealHeaderTool/Private/ClassMaps.h
+++ b/Engine/Source/Programs/UnrealHeaderTool/Private/ClassMaps.h
@@ -6,6 +6,7 @@
 #include "ParserClass.h"
 #include "Scope.h"
 #include "HeaderProvider.h"
+#include "UnderlyingEnumType.h"
 
 #include "UnrealSourceFile.h"
 #include "ClassDeclarationMetaData.h"
@@ -24,7 +25,7 @@ extern TSet<FUnrealSourceFile*> GPublicSourceFileSet;
 extern TMap<UProperty*, FString> GArrayDimensions;
 extern TMap<UPackage*,  const FManifestModule*> GPackageToManifestModuleMap;
 extern TMap<UField*, uint32> GGeneratedCodeCRCs;
-extern TMap<UEnum*,  EPropertyType> GEnumUnderlyingTypes;
+extern TMap<UEnum*, EUnderlyingEnumType> GEnumUnderlyingTypes;
 extern TMap<FName, TSharedRef<FClassDeclarationMetaData> > GClassDeclarations;
 extern TSet<UProperty*> GUnsizedProperties;
 
diff --git a/Engine/Source/Programs/UnrealHeaderTool/Private/CodeGenerator.cpp b/Engine/Source/Programs/UnrealHeaderTool/Private/CodeGenerator.cpp
--- a/Engine/Source/Programs/UnrealHeaderTool/Private/CodeGenerator.cpp
+++ b/Engine/Source/Programs/UnrealHeaderTool/Private/CodeGenerator.cpp
@@ -929,7 +929,7 @@ FString FNativeClassHeaderGenerator::GetOverriddenNameForLiteral(const UField* I
 	return TEXT("\"") + Item->GetName() + TEXT("\"");
 }
 
-FString FNativeClassHeaderGenerator::PropertyNew(FString& Meta, UProperty* Prop, const FString& OuterString, const FString& PropMacro, const TCHAR* NameSuffix, const TCHAR* Spaces, const TCHAR* SourceStruct)
+FString FNativeClassHeaderGenerator::PropertyNew(FString& Meta, UProperty* Prop, const FString& OuterString, const FString& PropMacro, const TCHAR* Name, const TCHAR* Spaces, const TCHAR* SourceStruct)
 {
 	FString ExtraArgs;
 	FString GeneratedCrc;
@@ -975,6 +975,10 @@ FString FNativeClassHeaderGenerator::PropertyNew(FString& Meta, UProperty* Prop,
 			ExtraArgs = FString::Printf(TEXT(", %s"), *GetSingletonName(ByteProperty->Enum));
 		}
 	}
+	else if (UEnumProperty* EnumProperty = Cast<UEnumProperty>(Prop))
+	{
+		ExtraArgs = FString::Printf(TEXT(", %s"), *GetSingletonName(EnumProperty->Enum));
+	}
 	else if (UBoolProperty* BoolProperty = Cast<UBoolProperty>(Prop))
 	{
 		if (Cast<UArrayProperty>(BoolProperty->GetOuter()) || Cast<UMapProperty>(BoolProperty->GetOuter()) || Cast<USetProperty>(BoolProperty->GetOuter()))
@@ -1026,23 +1030,22 @@ FString FNativeClassHeaderGenerator::PropertyNew(FString& Meta, UProperty* Prop,
 		*ExtraArgs);
 	TheFlagAudit.Add(Prop, TEXT("PropertyFlags"), Prop->PropertyFlags);
 
-	FString Symbol = FString::Printf(TEXT("NewProp_%s%s"), *Prop->GetName(), NameSuffix);
 	FString Lines  = FString::Printf(TEXT("%sUProperty* %s = %s%s\r\n"), 
 		Spaces, 
-		*Symbol, 
+		Name, 
 		*Constructor,
 		*GetGeneratedCodeCRCTag(Prop));
 
 	if (Prop->ArrayDim != 1)
 	{
-		Lines += FString::Printf(TEXT("%s%s->ArrayDim = CPP_ARRAY_DIM(%s, %s);\r\n"), Spaces, *Symbol, *PropNameDep, SourceStruct);
+		Lines += FString::Printf(TEXT("%s%s->ArrayDim = CPP_ARRAY_DIM(%s, %s);\r\n"), Spaces, Name, *PropNameDep, SourceStruct);
 	}
 
 	if (Prop->RepNotifyFunc != NAME_None)
 	{
-		Lines += FString::Printf(TEXT("%s%s->RepNotifyFunc = FName(TEXT(\"%s\"));\r\n"), Spaces, *Symbol, *Prop->RepNotifyFunc.ToString());
+		Lines += FString::Printf(TEXT("%s%s->RepNotifyFunc = FName(TEXT(\"%s\"));\r\n"), Spaces, Name, *Prop->RepNotifyFunc.ToString());
 	}
-	Meta += GetMetaDataCodeForObject(Prop, *Symbol, Spaces);
+	Meta += GetMetaDataCodeForObject(Prop, Name, Spaces);
 	return Lines;
 }
 
@@ -1100,6 +1103,8 @@ void FNativeClassHeaderGenerator::OutputProperty(FString& Meta, FOutputDevice& O
 {
 	FString PropName = Prop->GetName();
 
+	FString PropVariableName = FString::Printf(TEXT("NewProp_%s"), *PropName);
+
 	{
 		FString SourceStruct;
 		UFunction* Function = Cast<UFunction>(Prop->GetOuter());
@@ -1142,29 +1147,67 @@ void FNativeClassHeaderGenerator::OutputProperty(FString& Meta, FOutputDevice& O
 		{
 			PropMacroOuterClass = FString::Printf(TEXT("CPP_PROPERTY_BASE(%s, %s)"), *PropNameDep, *SourceStruct);
 		}
-		OutputDevice.Log(*PropertyNew(Meta, Prop, OuterString, PropMacroOuterClass, TEXT(""), Spaces, *SourceStruct));
+		OutputDevice.Log(*PropertyNew(Meta, Prop, OuterString, PropMacroOuterClass, *PropVariableName, Spaces, *SourceStruct));
 	}
 
+	// Map of enum class properties to their outer's variable name
+	TMap<UNumericProperty*, FString> UnderlyingEnumSuffixes;
+
 	if (UArrayProperty* ArrayProperty = Cast<UArrayProperty>(Prop))
 	{
 		FString InnerOuterString = FString::Printf(TEXT("NewProp_%s"), *PropName);
 		FString PropMacroOuterArray = TEXT("FObjectInitializer(), EC_CppProperty, 0");
-		OutputDevice.Log(*PropertyNew(Meta, ArrayProperty->Inner, InnerOuterString, PropMacroOuterArray, TEXT("_Inner"), Spaces));
+		FString InnerVariableName = FString::Printf(TEXT("NewProp_%s_Inner"), *ArrayProperty->Inner->GetName());
+		OutputDevice.Log(*PropertyNew(Meta, ArrayProperty->Inner, InnerOuterString, PropMacroOuterArray, *InnerVariableName, Spaces));
+
+		if (UEnumProperty* EnumArrayProperty = Cast<UEnumProperty>(ArrayProperty->Inner))
+		{
+			UnderlyingEnumSuffixes.Add(EnumArrayProperty->UnderlyingProp, MoveTemp(InnerVariableName));
+		}
 	}
 
 	else if (UMapProperty* MapProperty = Cast<UMapProperty>(Prop))
 	{
 		FString InnerOuterString = FString::Printf(TEXT("NewProp_%s"), *PropName);
 		FString PropMacroOuterMap = TEXT("FObjectInitializer(), EC_CppProperty, ");
-		OutputDevice.Log(*PropertyNew(Meta, MapProperty->KeyProp,   InnerOuterString, PropMacroOuterMap + TEXT("0"), TEXT("_KeyProp"),   Spaces));
-		OutputDevice.Log(*PropertyNew(Meta, MapProperty->ValueProp, InnerOuterString, PropMacroOuterMap + TEXT("1"), TEXT("_ValueProp"), Spaces));
+		FString KeyVariableName = FString::Printf(TEXT("NewProp_%s_KeyProp"), *MapProperty->KeyProp->GetName());
+		FString ValueVariableName = FString::Printf(TEXT("NewProp_%s_ValueProp"), *MapProperty->ValueProp->GetName());
+		OutputDevice.Log(*PropertyNew(Meta, MapProperty->KeyProp,   InnerOuterString, PropMacroOuterMap + TEXT("0"), *KeyVariableName,   Spaces));
+		OutputDevice.Log(*PropertyNew(Meta, MapProperty->ValueProp, InnerOuterString, PropMacroOuterMap + TEXT("1"), *ValueVariableName, Spaces));
+
+		if (UEnumProperty* EnumKeyProperty = Cast<UEnumProperty>(MapProperty->KeyProp))
+		{
+			UnderlyingEnumSuffixes.Add(EnumKeyProperty->UnderlyingProp, MoveTemp(KeyVariableName));
+		}
+
+		if (UEnumProperty* EnumValueProperty = Cast<UEnumProperty>(MapProperty->ValueProp))
+		{
+			UnderlyingEnumSuffixes.Add(EnumValueProperty->UnderlyingProp, MoveTemp(ValueVariableName));
+		}
 	}
 
 	else if (USetProperty* SetProperty = Cast<USetProperty>(Prop))
 	{
-		FString InnerOuterString = FString::Printf(TEXT("NewProp_%s"), *Prop->GetName());
-		FString PropMacroOuterArray = TEXT("FObjectInitializer(), EC_CppProperty, 0");
-		OutputDevice.Log(*PropertyNew(Meta, SetProperty->ElementProp, InnerOuterString, PropMacroOuterArray, TEXT("_ElementProp"), Spaces));
+		FString InnerOuterString = FString::Printf(TEXT("NewProp_%s"), *PropName);
+		FString PropMacroOuterSet = TEXT("FObjectInitializer(), EC_CppProperty, 0");
+		FString ElementVariableName = FString::Printf(TEXT("NewProp_%s_ElementProp"), *SetProperty->ElementProp->GetName());
+		OutputDevice.Log(*PropertyNew(Meta, SetProperty->ElementProp, InnerOuterString, PropMacroOuterSet, *ElementVariableName, Spaces));
+
+		if (UEnumProperty* EnumSetProperty = Cast<UEnumProperty>(SetProperty->ElementProp))
+		{
+			UnderlyingEnumSuffixes.Add(EnumSetProperty->UnderlyingProp, MoveTemp(ElementVariableName));
+		}
+	}
+
+	else if (UEnumProperty* EnumProperty = Cast<UEnumProperty>(Prop))
+	{
+		UnderlyingEnumSuffixes.Add(EnumProperty->UnderlyingProp, PropVariableName);
+	}
+
+	FString PropMacroOuterEnum = TEXT("FObjectInitializer(), EC_CppProperty, 0");
+	for (const TPair<UNumericProperty*, FString>& PropAndVarName : UnderlyingEnumSuffixes)
+	{
+		OutputDevice.Log(*PropertyNew(Meta, PropAndVarName.Key, PropAndVarName.Value, PropMacroOuterEnum, *(PropAndVarName.Value + TEXT("_Underlying")), Spaces));
 	}
 }
 
@@ -3179,12 +3222,12 @@ void FNativeClassHeaderGenerator::ExportGeneratedEnumsInitCode(const TArray<UEnu
 
 			const TCHAR* UEnumObjectFlags = bIsDynamic ? TEXT("RF_Public|RF_Transient") : TEXT("RF_Public|RF_Transient|RF_MarkAsNative");
 			GeneratedEnumRegisterFunctionText.Logf(TEXT("\t\t\tReturnEnum = new(EC_InternalUseOnlyConstructor, Outer, TEXT(\"%s\"), %s) UEnum(FObjectInitializer());\r\n"), *FNativeClassHeaderGenerator::GetOverriddenName(Enum), UEnumObjectFlags);
-			GeneratedEnumRegisterFunctionText.Logf(TEXT("\t\t\tTArray<TPair<FName, uint8>> EnumNames;\r\n"));
+			GeneratedEnumRegisterFunctionText.Logf(TEXT("\t\t\tTArray<TPair<FName, int64>> EnumNames;\r\n"));
 			for (int32 Index = 0; Index < Enum->NumEnums(); Index++)
 			{
 				const TCHAR* OverridenNameMetaDatakey = TEXT("OverrideName");
 				const FString KeyName = Enum->HasMetaData(OverridenNameMetaDatakey, Index) ? Enum->GetMetaData(OverridenNameMetaDatakey, Index) : Enum->GetNameByIndex(Index).ToString();
-				GeneratedEnumRegisterFunctionText.Logf(TEXT("\t\t\tEnumNames.Add(TPairInitializer<FName, uint8>(FName(TEXT(\"%s\")), %d));\r\n"), *KeyName, Enum->GetValueByIndex(Index));
+				GeneratedEnumRegisterFunctionText.Logf(TEXT("\t\t\tEnumNames.Add(TPairInitializer<FName, int64>(FName(TEXT(\"%s\")), %lld));\r\n"), *KeyName, Enum->GetValueByIndex(Index));
 			}
 
 			FString EnumTypeStr;
@@ -3704,7 +3747,7 @@ void ExportMCPDeclaration(FOutputDevice& Out, const FString& MessageName, TField
 			else
 			{
 				UByteProperty* ByteProperty = Cast<UByteProperty>(Property);
-				if (ByteProperty != NULL && ByteProperty->Enum != NULL)
+				if ((ByteProperty && ByteProperty->Enum) || Cast<UEnumProperty>(Property))
 				{
 					// treat enums like strings because that's how they'll be exported in JSON
 					MCPTypeName = ToMCPTypeMappings.Find(TEXT("FString"));
@@ -3951,6 +3994,12 @@ FString FNativeClassHeaderGenerator::GetNullParameterValue( UProperty* Prop, boo
 
 		return TEXT("0");
 	}
+	else if (PropClass == UEnumProperty::StaticClass())
+	{
+		UEnumProperty* EnumProp = (UEnumProperty*)Prop;
+
+		return FString::Printf(TEXT("(%s)0"), *EnumProp->Enum->GetName());
+	}
 	else if ( PropClass == UBoolProperty::StaticClass() )
 	{
 		return TEXT("false");
@@ -4327,7 +4376,6 @@ void FNativeClassHeaderGenerator::ExportFunctionThunk(FUHTStringBuilder& RPCWrap
 		FString EvalParameterText;				// e.g. (UObject*,NULL)
 
 		FString TypeText;
-		bool bPassAsNoPtr = false;
 
 		if (Param->ArrayDim > 1)
 		{
@@ -4351,11 +4399,10 @@ void FNativeClassHeaderGenerator::ExportFunctionThunk(FUHTStringBuilder& RPCWrap
 			}
 		}
 
-		if (Param->HasAllPropertyFlags(CPF_UObjectWrapper | CPF_OutParm) 
-			&& Param->IsA(UClassProperty::StaticClass()))
+		bool bPassAsNoPtr = Param->HasAllPropertyFlags(CPF_UObjectWrapper | CPF_OutParm) && Param->IsA(UClassProperty::StaticClass());
+		if (bPassAsNoPtr)
 		{
 			TypeText = Param->GetCPPType();
-			bPassAsNoPtr = true;
 		}
 
 		FUHTStringBuilder ReplacementText;
@@ -4372,7 +4419,7 @@ void FNativeClassHeaderGenerator::ExportFunctionThunk(FUHTStringBuilder& RPCWrap
 		{
 			if (!bPassAsNoPtr)
 			{
-			EvalModifierText += TEXT("_REF");
+				EvalModifierText += TEXT("_REF");
 			}
 			else
 			{
@@ -4423,24 +4470,34 @@ void FNativeClassHeaderGenerator::ExportFunctionThunk(FUHTStringBuilder& RPCWrap
 			}
 		}
 
+		UEnum* Enum = nullptr;
 		UByteProperty* ByteProp = Cast<UByteProperty>(Param);
 		if (ByteProp && ByteProp->Enum)
+		{
+			Enum = ByteProp->Enum;
+		}
+		else if (Param->IsA<UEnumProperty>())
+		{
+			Enum = ((UEnumProperty*)Param)->Enum;
+		}
+
+		if (Enum)
 		{
 			// For enums, add an explicit conversion 
-			if (!(ByteProp->PropertyFlags & CPF_OutParm))
+			if (!(Param->PropertyFlags & CPF_OutParm))
 			{
-				ParamName = FString::Printf(TEXT("%s(%s)"), *ByteProp->Enum->CppType, *ParamName);
+				ParamName = FString::Printf(TEXT("%s(%s)"), *Enum->CppType, *ParamName);
 			}
 			else
 			{
-				if (ByteProp->Enum->GetCppForm() == UEnum::ECppForm::EnumClass)
+				if (Enum->GetCppForm() == UEnum::ECppForm::EnumClass)
 				{
 					// If we're an enum class don't require the wrapper
-					ParamName = FString::Printf(TEXT("(%s&)(%s)"), *ByteProp->Enum->CppType, *ParamName);
+					ParamName = FString::Printf(TEXT("(%s&)(%s)"), *Enum->CppType, *ParamName);
 				}
 				else
 				{
-					ParamName = FString::Printf(TEXT("(TEnumAsByte<%s>&)(%s)"), *ByteProp->Enum->CppType, *ParamName);
+					ParamName = FString::Printf(TEXT("(TEnumAsByte<%s>&)(%s)"), *Enum->CppType, *ParamName);
 				}
 			}
 		}
@@ -4832,11 +4889,15 @@ TArray<UFunction*> FNativeClassHeaderGenerator::ExportCallbackFunctions(FUnrealS
 void FNativeClassHeaderGenerator::ApplyAlternatePropertyExportText(UProperty* Prop, FUHTStringBuilder& PropertyText)
 {
 	UArrayProperty* ArrayProperty = Cast<UArrayProperty>(Prop);
-	UByteProperty* InnerByteProperty = ArrayProperty ? Cast<UByteProperty>(ArrayProperty->Inner) : nullptr;
-	if (InnerByteProperty && InnerByteProperty->Enum && FClass::IsDynamic(InnerByteProperty->Enum))
+	UProperty* InnerProperty = ArrayProperty ? ArrayProperty->Inner : nullptr;
+	if (InnerProperty && (
+			InnerProperty->IsA<UByteProperty>() && ((UByteProperty*)InnerProperty)->Enum && FClass::IsDynamic(((UByteProperty*)InnerProperty)->Enum) ||
+			InnerProperty->IsA<UEnumProperty>()                                          && FClass::IsDynamic(((UEnumProperty*)InnerProperty)->Enum)
+		)
+	)
 	{
-		const FString Original = InnerByteProperty->GetCPPType();
-		const FString RawByte = InnerByteProperty->GetCPPType(nullptr, EPropertyExportCPPFlags::CPPF_BlueprintCppBackend);
+		const FString Original = InnerProperty->GetCPPType();
+		const FString RawByte = InnerProperty->GetCPPType(nullptr, EPropertyExportCPPFlags::CPPF_BlueprintCppBackend);
 		if (Original != RawByte)
 		{
 			PropertyText.ReplaceInline(*Original, *RawByte, ESearchCase::CaseSensitive);
@@ -5932,7 +5993,12 @@ ECompilationResult::Type UnrealHeaderTool_Main(const FString& ModuleInfoFilename
 	if (bUseMakefile)
 	{
 		MakefilePath = FPaths::Combine(*ModuleInfoPath, TEXT("UHT.makefile"));
-		UHTMakefile.LoadFromFile(*MakefilePath, &GManifest);
+
+		// If makefile failed to load, clear it
+		if (!UHTMakefile.LoadFromFile(*MakefilePath, &GManifest))
+		{
+			UHTMakefile = FUHTMakefile();
+		}
 	}
 	UHTMakefile.StartPreloading();
 	{
diff --git a/Engine/Source/Programs/UnrealHeaderTool/Private/HeaderParser.cpp b/Engine/Source/Programs/UnrealHeaderTool/Private/HeaderParser.cpp
--- a/Engine/Source/Programs/UnrealHeaderTool/Private/HeaderParser.cpp
+++ b/Engine/Source/Programs/UnrealHeaderTool/Private/HeaderParser.cpp
@@ -506,6 +506,41 @@ namespace
 
 	UProperty* CreateVariableProperty(FPropertyBase& VarProperty, UObject* Scope, FName Name, EObjectFlags ObjectFlags, EVariableCategory::Type VariableCategory, FUHTMakefile& UHTMakefile, FUnrealSourceFile* UnrealSourceFile)
 	{
+		// Check if it's an enum class property
+		if (const EUnderlyingEnumType* EnumPropType = GEnumUnderlyingTypes.Find(VarProperty.Enum))
+		{
+			FPropertyBase UnderlyingProperty = VarProperty;
+			UnderlyingProperty.Enum = nullptr;
+			switch (*EnumPropType)
+			{
+				case EUnderlyingEnumType::int8:        UnderlyingProperty.Type = CPT_Int8;   break;
+				case EUnderlyingEnumType::int16:       UnderlyingProperty.Type = CPT_Int16;  break;
+				case EUnderlyingEnumType::int32:       UnderlyingProperty.Type = CPT_Int;    break;
+				case EUnderlyingEnumType::int64:       UnderlyingProperty.Type = CPT_Int64;  break;
+				case EUnderlyingEnumType::uint8:       UnderlyingProperty.Type = CPT_Byte;   break;
+				case EUnderlyingEnumType::uint16:      UnderlyingProperty.Type = CPT_UInt16; break;
+				case EUnderlyingEnumType::uint32:      UnderlyingProperty.Type = CPT_UInt32; break;
+				case EUnderlyingEnumType::uint64:      UnderlyingProperty.Type = CPT_UInt64; break;
+				case EUnderlyingEnumType::Unspecified: UnderlyingProperty.Type = CPT_Int;    break;
+
+				default:
+					check(false);
+			}
+
+			if (*EnumPropType == EUnderlyingEnumType::Unspecified)
+			{
+				UnderlyingProperty.IntType = EIntType::Unsized;
+			}
+
+			UEnumProperty* Result = new (EC_InternalUseOnlyConstructor, Scope, Name, ObjectFlags) UEnumProperty(FObjectInitializer());
+			UNumericProperty* UnderlyingProp = CastChecked<UNumericProperty>(CreateVariableProperty(UnderlyingProperty, Result, TEXT("UnderlyingType"), ObjectFlags, VariableCategory, UHTMakefile, UnrealSourceFile));
+			Result->UnderlyingProp = UnderlyingProp;
+			Result->Enum = VarProperty.Enum;
+
+			UHTMakefile.AddEnumProperty(UnrealSourceFile, Result);
+			return Result;
+		}
+
 		switch (VarProperty.Type)
 		{
 			case CPT_Byte:
@@ -910,7 +945,8 @@ namespace
 			|| Property->IsA<UBoolProperty>()
 			|| Property->IsA<UStrProperty>()
 			|| Property->IsA<UTextProperty>()
-			|| Property->IsA<UDelegateProperty>();
+			|| Property->IsA<UDelegateProperty>()
+			|| Property->IsA<UEnumProperty>();
 
 		const bool bIsSupportedMemberVariable = Property->IsA<UWeakObjectProperty>() || Property->IsA<UMulticastDelegateProperty>();
 
@@ -1205,6 +1241,7 @@ UEnum* FHeaderParser::CompileEnum()
 	ValidateMetaDataFormat(Enum, EnumToken.MetaData);
 
 	// Read base for enum class
+	EUnderlyingEnumType UnderlyingType = EUnderlyingEnumType::uint8;
 	if (CppForm == UEnum::ECppForm::EnumClass)
 	{
 		if (MatchSymbol(TEXT(":")))
@@ -1215,21 +1252,56 @@ UEnum* FHeaderParser::CompileEnum()
 				FError::Throwf(TEXT("Missing enum base") );
 			}
 
-			// We only support uint8 at the moment, until the properties get updated
-			if (FCString::Strcmp(BaseToken.Identifier, TEXT("uint8")))
+			if (!FCString::Strcmp(BaseToken.Identifier, TEXT("uint8")))
 			{
-				FError::Throwf(TEXT("Only enum bases of type uint8 are currently supported"));
+				UnderlyingType = EUnderlyingEnumType::uint8;
+			}
+			else if (!FCString::Strcmp(BaseToken.Identifier, TEXT("uint16")))
+			{
+				UnderlyingType = EUnderlyingEnumType::uint16;
+			}
+			else if (!FCString::Strcmp(BaseToken.Identifier, TEXT("uint32")))
+			{
+				UnderlyingType = EUnderlyingEnumType::uint32;
+			}
+			else if (!FCString::Strcmp(BaseToken.Identifier, TEXT("uint64")))
+			{
+				UnderlyingType = EUnderlyingEnumType::uint64;
+			}
+			else if (!FCString::Strcmp(BaseToken.Identifier, TEXT("int8")))
+			{
+				UnderlyingType = EUnderlyingEnumType::int8;
+			}
+			else if (!FCString::Strcmp(BaseToken.Identifier, TEXT("int16")))
+			{
+				UnderlyingType = EUnderlyingEnumType::int16;
+			}
+			else if (!FCString::Strcmp(BaseToken.Identifier, TEXT("int32")))
+			{
+				UnderlyingType = EUnderlyingEnumType::int32;
+			}
+			else if (!FCString::Strcmp(BaseToken.Identifier, TEXT("int64")))
+			{
+				UnderlyingType = EUnderlyingEnumType::int64;
+			}
+			else
+			{
+				FError::Throwf(TEXT("Unsupported enum class base type: %s"), BaseToken.Identifier);
 			}
-
-			GEnumUnderlyingTypes.Add(Enum, CPT_Byte);
-			UHTMakefile.AddGEnumUnderlyingType(CurrentSrcFile, Enum, CPT_Byte);
 		}
-	#if DEPRECATE_ENUM_AS_BYTE_FOR_ENUM_CLASSES
 		else
 		{
-			FError::Throwf(TEXT("Missing base specifier for enum class '%s' - did you mean ': uint8'?"), EnumToken.Identifier);
+			UnderlyingType = EUnderlyingEnumType::Unspecified;
 		}
-	#endif
+
+		GEnumUnderlyingTypes.Add(Enum, UnderlyingType);
+		UHTMakefile.AddGEnumUnderlyingType(CurrentSrcFile, Enum, UnderlyingType);
+	}
+
+	static const FName BlueprintTypeName = TEXT("BlueprintType");
+	if (UnderlyingType != EUnderlyingEnumType::uint8 && !EnumToken.MetaData.Contains(BlueprintTypeName))
+	{
+		FError::Throwf(TEXT("Invalid BlueprintType enum base - currently only uint8 supported"));
 	}
 
 	// Get opening brace.
@@ -1272,9 +1344,9 @@ UEnum* FHeaderParser::CompileEnum()
 	FToken TagToken;
 
 	TArray<FScriptLocation> EnumTagLocations;
-	TArray<TPair<FName, uint8>> EnumNames;
+	TArray<TPair<FName, int64>> EnumNames;
 
-	int32 CurrentEnumValue = 0;
+	int64 CurrentEnumValue = 0;
 
 	while (GetIdentifier(TagToken))
 	{
@@ -1285,12 +1357,58 @@ UEnum* FHeaderParser::CompileEnum()
 		// Try to read an optional explicit enum value specification
 		if (MatchSymbol(TEXT("=")))
 		{
-			int32 NewEnumValue = 0;
-			GetConstInt(/*out*/ NewEnumValue, TEXT("Enumerator value"));
+			int64 NewEnumValue = 0;
+			GetConstInt64(/*out*/ NewEnumValue, TEXT("Enumerator value"));
 
-			if ((NewEnumValue < CurrentEnumValue) || (NewEnumValue > 255))
+			if (EnumNames.Num() > 0 && NewEnumValue < CurrentEnumValue)
 			{
-				FError::Throwf(TEXT("Explicitly specified enum values must be greater than any previous value and less than 256"));
+				FError::Throwf(TEXT("Explicitly specified enum value (%d) must be greater than the previous value (%d)"), NewEnumValue, CurrentEnumValue);
+			}
+
+			if (UnderlyingType == EUnderlyingEnumType::Unspecified || UnderlyingType == EUnderlyingEnumType::int8 || UnderlyingType == EUnderlyingEnumType::int16 || UnderlyingType == EUnderlyingEnumType::int32 || UnderlyingType == EUnderlyingEnumType::int64)
+			{
+				int64 Min = 0;
+				int64 Max = 0;
+				switch (UnderlyingType)
+				{
+					case EUnderlyingEnumType::Unspecified: Min = TNumericLimits<int>  ::Min(); Max = TNumericLimits<int>  ::Max(); break;
+					case EUnderlyingEnumType::int8:        Min = TNumericLimits<int8> ::Min(); Max = TNumericLimits<int8> ::Max(); break;
+					case EUnderlyingEnumType::int16:       Min = TNumericLimits<int16>::Min(); Max = TNumericLimits<int16>::Max(); break;
+					case EUnderlyingEnumType::int32:       Min = TNumericLimits<int32>::Min(); Max = TNumericLimits<int32>::Max(); break;
+					case EUnderlyingEnumType::int64:       Min = TNumericLimits<int64>::Min(); Max = TNumericLimits<int64>::Max(); break;
+
+					default:
+						check(false);
+				}
+
+				if (NewEnumValue < Min || NewEnumValue > Max)
+				{
+					FError::Throwf(TEXT("Explicitly specified enum value (%lld) must be in the range of the underlying type of the enum (%lld to %lld)"), NewEnumValue, Min, Max);
+				}
+			}
+			else
+			{
+				uint64 Min = 0;
+				uint64 Max = 0;
+				switch (UnderlyingType)
+				{
+					case EUnderlyingEnumType::uint8:  Min = TNumericLimits<uint8> ::Min(); Max = TNumericLimits<uint8> ::Max(); break;
+					case EUnderlyingEnumType::uint16: Min = TNumericLimits<uint16>::Min(); Max = TNumericLimits<uint16>::Max(); break;
+					case EUnderlyingEnumType::uint32: Min = TNumericLimits<uint32>::Min(); Max = TNumericLimits<uint32>::Max(); break;
+					case EUnderlyingEnumType::uint64: Min = TNumericLimits<uint64>::Min(); Max = TNumericLimits<uint64>::Max(); break;
+
+					default:
+						check(false);
+				}
+
+				if (NewEnumValue < 0)
+				{
+					FError::Throwf(TEXT("Explicitly specified enum value (%lld) must be in the range of the underlying type of the enum (%llu to %llu)"), NewEnumValue, Min, Max);
+				}
+				else if ((uint64)NewEnumValue < Min || (uint64)NewEnumValue > Max)
+				{
+					FError::Throwf(TEXT("Explicitly specified enum value (%llu) must be in the range of the underlying type of the enum (%llu to %llu)"), (uint64)NewEnumValue, Min, Max);
+				}
 			}
 
 			CurrentEnumValue = NewEnumValue;
@@ -1314,18 +1432,13 @@ UEnum* FHeaderParser::CompileEnum()
 			break;
 		}
 
-		TPair<FName, uint8> CurrentEnum = TPair<FName, uint8>(TPairInitializer<FName, uint8>(NewTag, CurrentEnumValue));
+		TPair<FName, int64> CurrentEnum = TPair<FName, int64>(TPairInitializer<FName, int64>(NewTag, CurrentEnumValue));
 
 		if (EnumNames.Find(CurrentEnum, iFound))
 		{
 			FError::Throwf(TEXT("Duplicate enumeration tag %s"), TagToken.Identifier );
 		}
 
-		if (CurrentEnumValue > 255)
-		{
-			FError::Throwf(TEXT("Exceeded maximum of 255 enumerators") );
-		}
-
 		UEnum* FoundEnum = NULL;
 		if (UEnum::LookupEnumName(NewTag, &FoundEnum) != INDEX_NONE)
 		{
@@ -1915,10 +2028,6 @@ UScriptStruct* FHeaderParser::CompileStructDeclaration(FClasses& AllClasses)
 						FError::Throwf(TEXT("Parent Struct '%s' is missing a valid Unreal prefix, expecting '%s'"), *ParentStructNameInScript, *FString::Printf(TEXT("%s%s"), PrefixCPP, *Type->GetName()));
 					}
 				}
-				else
-				{
-				
-				}
 			}
 		}
 		else
@@ -2566,7 +2675,7 @@ void FHeaderParser::VerifyRepNotifyCallbacks( UClass* TargetClass )
 
 				if ( TargetFunc->NumParms >= 2 && Parm)
 				{
-					// A 2nd parmaeter for arrays can be specified as a const TArray<uint8>&. This is a list of element indices that have changed
+					// A 2nd parameter for arrays can be specified as a const TArray<uint8>&. This is a list of element indices that have changed
 					UArrayProperty *ArrayProp = Cast<UArrayProperty>(*Parm);
 					if (!(ArrayProp && Cast<UByteProperty>(ArrayProp->Inner)) || !(Parm->GetPropertyFlags() & CPF_ConstParm) || !(Parm->GetPropertyFlags() & CPF_ReferenceParm))
 					{
@@ -2619,7 +2728,7 @@ void FHeaderParser::CompileDirective(FClasses& AllClasses)
 	else if (Directive.Matches(TEXT("linenumber")))
 	{
 		FToken Number;
-		if (!GetToken(Number) || (Number.TokenType != TOKEN_Const) || (Number.Type != CPT_Int))
+		if (!GetToken(Number) || (Number.TokenType != TOKEN_Const) || (Number.Type != CPT_Int && Number.Type != CPT_Int64))
 		{
 			FError::Throwf(TEXT("Missing line number in line number directive"));
 		}
@@ -3599,10 +3708,10 @@ FIndexRange*                    ParsedVarIndexRange
 
 		if (VariableCategory == EVariableCategory::Member)
 		{
-			auto* EnumUnderlyingType = GEnumUnderlyingTypes.Find(Enum);
-			if (!EnumUnderlyingType || *EnumUnderlyingType != CPT_Byte)
+			EUnderlyingEnumType* EnumUnderlyingType = GEnumUnderlyingTypes.Find(Enum);
+			if (!EnumUnderlyingType)
 			{
-				FError::Throwf(TEXT("You cannot use the raw enum name as a type for member variables, instead use TEnumAsByte or a C++11 enum class with an explicit underlying type (currently only uint8 supported)."), *Enum->CppType);
+				FError::Throwf(TEXT("You cannot use the raw enum name as a type for member variables, instead use TEnumAsByte or a C++11 enum class with an explicit underlying type."), *Enum->CppType);
 			}
 		}
 
@@ -7963,6 +8072,22 @@ bool FHeaderParser::DefaultValueStringCppFormatToInnerFormat(const UProperty* Pr
 				return ( 0 <= Value ) && ( 255 >= Value );
 			}
 		}
+		else if( Property->IsA(UEnumProperty::StaticClass()) )
+		{
+			const UEnumProperty* EnumProp = CastChecked<UEnumProperty>(Property);
+			if (const UEnum* Enum = CastChecked<UEnumProperty>(Property)->GetEnum())
+			{
+				OutForm = FDefaultValueHelper::GetUnqualifiedEnumValue(FDefaultValueHelper::RemoveWhitespaces(CppForm));
+				return Enum->FindEnumIndex(*OutForm) != INDEX_NONE;
+			}
+
+			int64 Value;
+			if (FDefaultValueHelper::ParseInt64(CppForm, Value))
+			{
+				OutForm = LexicalConversion::ToString(Value);
+				return EnumProp->GetUnderlyingProperty()->CanHoldValue(Value);
+			}
+		}
 		else if( Property->IsA(UFloatProperty::StaticClass()) )
 		{
 			float Value;
diff --git a/Engine/Source/Programs/UnrealHeaderTool/Private/NativeClassExporter.h b/Engine/Source/Programs/UnrealHeaderTool/Private/NativeClassExporter.h
--- a/Engine/Source/Programs/UnrealHeaderTool/Private/NativeClassExporter.h
+++ b/Engine/Source/Programs/UnrealHeaderTool/Private/NativeClassExporter.h
@@ -469,11 +469,11 @@ private:
 	 * @param	Prop			Property to export
 	 * @param	OuterString		String that specifies the outer to add the properties to
 	 * @param	PropMacro		String specifying the macro to call to get the property offset
-	 * @param	NameSuffix		Suffix for the generated variable name
+	 * @param	Name			Name for the generated variable
 	 * @param	Spaces			String of spaces to use as an indent
 	 * @param	SourceStruct	Structure that the property offset is relative to
 	**/
-	FString PropertyNew(FString& Meta, UProperty* Prop, const FString& OuterString, const FString& PropMacro, const TCHAR* NameSuffix, const TCHAR* Spaces, const TCHAR* SourceStruct = NULL);
+	FString PropertyNew(FString& Meta, UProperty* Prop, const FString& OuterString, const FString& PropMacro, const TCHAR* Name, const TCHAR* Spaces, const TCHAR* SourceStruct = NULL);
 
 	/**
 	 * Exports the proxy definitions for the list of enums specified
diff --git a/Engine/Source/Programs/UnrealHeaderTool/Private/ParserHelper.cpp b/Engine/Source/Programs/UnrealHeaderTool/Private/ParserHelper.cpp
--- a/Engine/Source/Programs/UnrealHeaderTool/Private/ParserHelper.cpp
+++ b/Engine/Source/Programs/UnrealHeaderTool/Private/ParserHelper.cpp
@@ -258,17 +258,17 @@ bool FFunctionData::TryFindForFunction(UFunction* Function, FFunctionData*& OutD
 
 FClassMetaData* FCompilerMetadataManager::AddClassData(UStruct* Struct, FUHTMakefile& UHTMakefile, FUnrealSourceFile* UnrealSourceFile)
 {
-	TScopedPointer<FClassMetaData>* pClassData = Find(Struct);
+	TUniquePtr<FClassMetaData>* pClassData = Find(Struct);
 	if (pClassData == NULL)
 	{
 		pClassData = &Emplace(Struct, new FClassMetaData());
 		if (UnrealSourceFile)
 		{
-			UHTMakefile.AddClassMetaData(UnrealSourceFile, *pClassData);
+			UHTMakefile.AddClassMetaData(UnrealSourceFile, MakeUnique<FClassMetaData>(*pClassData->Get()));
 		}
 	}
 
-	return *pClassData;
+	return pClassData->Get();
 }
 
 FTokenData* FPropertyData::Set(UProperty* InKey, const FTokenData& InValue, FUHTMakefile& UHTMakefile, FUnrealSourceFile* UnrealSourceFile)
diff --git a/Engine/Source/Programs/UnrealHeaderTool/Private/ParserHelper.h b/Engine/Source/Programs/UnrealHeaderTool/Private/ParserHelper.h
--- a/Engine/Source/Programs/UnrealHeaderTool/Private/ParserHelper.h
+++ b/Engine/Source/Programs/UnrealHeaderTool/Private/ParserHelper.h
@@ -3,6 +3,7 @@
 #pragma once
 #include "ClassMaps.h"
 #include "UniqueObj.h"
+#include "UniquePtr.h"
 
 extern class FCompilerMetadataManager GScriptHelper;
 
@@ -278,6 +279,26 @@ public:
 			Enum = Cast<UByteProperty>(Property)->Enum;
 			IntType = EIntType::Sized;
 		}
+		else if( ClassOfProperty==UEnumProperty::StaticClass() )
+		{
+			UEnumProperty* EnumProp = Cast<UEnumProperty>(Property);
+			UNumericProperty* UnderlyingProp = EnumProp->GetUnderlyingProperty();
+
+			*this = FPropertyBase(
+				  UnderlyingProp->IsA<UInt8Property>()   ? CPT_Int8
+				: UnderlyingProp->IsA<UInt16Property>()  ? CPT_Int16
+				: UnderlyingProp->IsA<UIntProperty>()    ? CPT_Int
+				: UnderlyingProp->IsA<UInt64Property>()  ? CPT_Int64
+				: UnderlyingProp->IsA<UByteProperty>()   ? CPT_Byte
+				: UnderlyingProp->IsA<UUInt16Property>() ? CPT_UInt16
+				: UnderlyingProp->IsA<UUInt32Property>() ? CPT_UInt32
+				: UnderlyingProp->IsA<UUInt64Property>() ? CPT_UInt64
+				:                                          CPT_None
+			);
+			check(this->Type != CPT_None);
+			Enum = EnumProp->Enum;
+			IntType = EIntType::Sized;
+		}
 		else if( ClassOfProperty==UInt8Property::StaticClass() )
 		{
 			*this = FPropertyBase(CPT_Int8);
@@ -706,6 +727,7 @@ public:
 	{
 		// TOKEN_Const values.
 		uint8 Byte;								// If CPT_Byte.
+		int64 Int64;							// If CPT_Int64.
 		int32 Int;								// If CPT_Int.
 		bool NativeBool;						// if CPT_Bool
 		float Float;							// If CPT_Float.
@@ -722,7 +744,7 @@ public:
 	 */
 	void Clone( const FToken& Other );
 
-	FString GetConstantValue()
+	FString GetConstantValue() const
 	{
 		if (TokenType == TOKEN_Const)
 		{
@@ -730,6 +752,8 @@ public:
 			{
 				case CPT_Byte:
 					return FString::Printf(TEXT("%u"), Byte);
+				case CPT_Int64:
+					return FString::Printf(TEXT("%ld"), Int64);
 				case CPT_Int:
 					return FString::Printf(TEXT("%i"), Int);
 				case CPT_Bool:
@@ -749,7 +773,6 @@ public:
 				case CPT_Rotation:
 				case CPT_Int8:
 				case CPT_Int16:
-				case CPT_Int64:
 				case CPT_UInt16:
 				case CPT_UInt32:
 				case CPT_UInt64:
@@ -832,6 +855,12 @@ public:
 
 	// Setters.
 	
+	void SetConstInt64( int64 InInt64 )
+	{
+		(FPropertyBase&)*this = FPropertyBase(CPT_Int64);
+		Int64			= InInt64;
+		TokenType		= TOKEN_Const;
+	}
 	void SetConstInt( int32 InInt )
 	{
 		(FPropertyBase&)*this = FPropertyBase(CPT_Int);
@@ -885,7 +914,42 @@ public:
 	// Getters.
 	bool GetConstInt( int32& I ) const
 	{
-		if( TokenType==TOKEN_Const && Type==CPT_Int )
+		if( TokenType==TOKEN_Const && Type==CPT_Int64 )
+		{
+			I = Int64;
+			return 1;
+		}
+		else if( TokenType==TOKEN_Const && Type==CPT_Int )
+		{
+			I = Int;
+			return 1;
+		}
+		else if( TokenType==TOKEN_Const && Type==CPT_Byte )
+		{
+			I = Byte;
+			return 1;
+		}
+		else if( TokenType==TOKEN_Const && Type==CPT_Float && Float==FMath::TruncToInt(Float))
+		{
+			I = (int32) Float;
+			return 1;
+		}
+		else if (TokenType == TOKEN_Const && Type == CPT_Double && Float == FMath::TruncToInt(Double))
+		{
+			I = (int32) Double;
+			return 1;
+		}
+		else return 0;
+	}
+
+	bool GetConstInt64( int64& I ) const
+	{
+		if( TokenType==TOKEN_Const && Type==CPT_Int64 )
+		{
+			I = Int64;
+			return 1;
+		}
+		else if( TokenType==TOKEN_Const && Type==CPT_Int )
 		{
 			I = Int;
 			return 1;
@@ -1687,7 +1751,7 @@ public:
  * The type of data tracked by this class is data that would otherwise only be accessible by adding a 
  * member property to UFunction/UProperty.  
  */
-class FCompilerMetadataManager : protected TMap<UStruct*, TScopedPointer<FClassMetaData> >
+class FCompilerMetadataManager : protected TMap<UStruct*, TUniquePtr<FClassMetaData> >
 {
 public:
 	/**
@@ -1708,12 +1772,12 @@ public:
 	 */
 	FClassMetaData* FindClassData(UStruct* Struct)
 	{
-		FClassMetaData* Result = NULL;
+		FClassMetaData* Result = nullptr;
 
-		TScopedPointer<FClassMetaData>* pClassData = Find(Struct);
+		TUniquePtr<FClassMetaData>* pClassData = Find(Struct);
 		if (pClassData)
 		{
-			Result = pClassData->GetOwnedPointer();
+			Result = pClassData->Get();
 		}
 
 		return Result;
@@ -1724,10 +1788,10 @@ public:
 	 */
 	void Shrink()
 	{
-		TMap<UStruct*, TScopedPointer<FClassMetaData> >::Shrink();
-		for (TMap<UStruct*, TScopedPointer<FClassMetaData> >::TIterator It(*this); It; ++It)
+		TMap<UStruct*, TUniquePtr<FClassMetaData> >::Shrink();
+		for (TMap<UStruct*, TUniquePtr<FClassMetaData> >::TIterator It(*this); It; ++It)
 		{
-			FClassMetaData* MetaData = It.Value();
+			FClassMetaData* MetaData = It->Value.Get();
 			MetaData->Shrink();
 		}
 	}
diff --git a/Engine/Source/Programs/UnrealHeaderTool/Private/UHTMakefile/ClassMetaDataArchiveProxy.cpp b/Engine/Source/Programs/UnrealHeaderTool/Private/UHTMakefile/ClassMetaDataArchiveProxy.cpp
--- a/Engine/Source/Programs/UnrealHeaderTool/Private/UHTMakefile/ClassMetaDataArchiveProxy.cpp
+++ b/Engine/Source/Programs/UnrealHeaderTool/Private/UHTMakefile/ClassMetaDataArchiveProxy.cpp
@@ -33,11 +33,11 @@ void FClassMetaDataArchiveProxy::AddReferencedNames(const FClassMetaData* ClassM
 	}
 }
 
-TScopedPointer<FClassMetaData> FClassMetaDataArchiveProxy::CreateClassMetaData() const
+TUniquePtr<FClassMetaData> FClassMetaDataArchiveProxy::CreateClassMetaData() const
 {
-	FClassMetaData* ClassMetaData = new FClassMetaData();
-	PostConstruct(ClassMetaData);
-	return TScopedPointer<FClassMetaData>(ClassMetaData);
+	TUniquePtr<FClassMetaData> ClassMetaData = MakeUnique<FClassMetaData>();
+	PostConstruct(ClassMetaData.Get());
+	return ClassMetaData;
 }
 
 void FClassMetaDataArchiveProxy::PostConstruct(FClassMetaData* ClassMetaData) const
diff --git a/Engine/Source/Programs/UnrealHeaderTool/Private/UHTMakefile/ClassMetaDataArchiveProxy.h b/Engine/Source/Programs/UnrealHeaderTool/Private/UHTMakefile/ClassMetaDataArchiveProxy.h
--- a/Engine/Source/Programs/UnrealHeaderTool/Private/UHTMakefile/ClassMetaDataArchiveProxy.h
+++ b/Engine/Source/Programs/UnrealHeaderTool/Private/UHTMakefile/ClassMetaDataArchiveProxy.h
@@ -4,6 +4,7 @@
 #include "UHTMakefile/MultipleInheritanceBaseClassArchiveProxy.h"
 #include "UHTMakefile/PropertyDataArchiveProxy.h"
 #include "UHTMakefile/StructDataArchiveProxy.h"
+#include "UniquePtr.h"
 
 class FUHTMakefile;
 class FClassMetaData;
@@ -17,7 +18,7 @@ struct FClassMetaDataArchiveProxy
 	static void AddReferencedNames(const FClassMetaData* FClassMetaData, FUHTMakefile& UHTMakefile);
 
 	friend FArchive& operator<<(FArchive& Ar, FClassMetaDataArchiveProxy& ClassMetaDataArchiveProxy);
-	TScopedPointer<FClassMetaData> CreateClassMetaData() const;
+	TUniquePtr<FClassMetaData> CreateClassMetaData() const;
 	void PostConstruct(FClassMetaData* ClassMetaData) const;
 	void Resolve(FClassMetaData* ClassMetaData, FUHTMakefile& UHTMakefile);
 
diff --git a/Engine/Source/Programs/UnrealHeaderTool/Private/UHTMakefile/CompilerMetadataManagerArchiveProxy.cpp b/Engine/Source/Programs/UnrealHeaderTool/Private/UHTMakefile/CompilerMetadataManagerArchiveProxy.cpp
--- a/Engine/Source/Programs/UnrealHeaderTool/Private/UHTMakefile/CompilerMetadataManagerArchiveProxy.cpp
+++ b/Engine/Source/Programs/UnrealHeaderTool/Private/UHTMakefile/CompilerMetadataManagerArchiveProxy.cpp
@@ -15,7 +15,7 @@ void FCompilerMetadataManagerArchiveProxy::AddReferencedNames(const FCompilerMet
 {
 	for (auto& Kvp : *CompilerMetadataManager)
 	{
-		FClassMetaDataArchiveProxy::AddReferencedNames(Kvp.Value, UHTMakefile);
+		FClassMetaDataArchiveProxy::AddReferencedNames(Kvp.Value.Get(), UHTMakefile);
 	}
 }
 
diff --git a/Engine/Source/Programs/UnrealHeaderTool/Private/UHTMakefile/EnumArchiveProxy.cpp b/Engine/Source/Programs/UnrealHeaderTool/Private/UHTMakefile/EnumArchiveProxy.cpp
--- a/Engine/Source/Programs/UnrealHeaderTool/Private/UHTMakefile/EnumArchiveProxy.cpp
+++ b/Engine/Source/Programs/UnrealHeaderTool/Private/UHTMakefile/EnumArchiveProxy.cpp
@@ -9,12 +9,12 @@ FEnumArchiveProxy::FEnumArchiveProxy(FUHTMakefile& UHTMakefile, const UEnum* Enu
 	EnumFlags = static_cast<uint32>(Enum->GetFlags());
 	CppForm = static_cast<uint32>(Enum->GetCppForm());
 
-	EnumValues.Empty(Enum->NumEnums());
+	EnumValuesX.Empty(Enum->NumEnums());
 	for (int32 i = 0; i < Enum->NumEnums(); ++i)
 	{
 		FName EnumName = Enum->GetNameByIndex(i);
-		uint8 Value = Enum->GetValueByIndex(i);
-		EnumValues.Add(TPairInitializer<FNameArchiveProxy, uint8>(FNameArchiveProxy(UHTMakefile, EnumName), Value));
+		int64 Value = Enum->GetValueByIndex(i);
+		EnumValuesX.Add(TPairInitializer<FNameArchiveProxy, int64>(FNameArchiveProxy(UHTMakefile, EnumName), Value));
 	}
 
 	CppType = Enum->CppType;
@@ -32,11 +32,11 @@ UEnum* FEnumArchiveProxy::CreateEnum(const FUHTMakefile& UHTMakefile) const
 
 void FEnumArchiveProxy::PostConstruct(UEnum* Enum, const FUHTMakefile& UHTMakefile) const
 {
-	TArray<TPair<FName, uint8>> Names;
-	Names.Empty(EnumValues.Num());
-	for (auto& Kvp : EnumValues)
+	TArray<TPair<FName, int64>> Names;
+	Names.Empty(EnumValuesX.Num());
+	for (auto& Kvp : EnumValuesX)
 	{
-		Names.Add(TPairInitializer<FName, uint8>(Kvp.Key.CreateName(UHTMakefile), Kvp.Value));
+		Names.Add(TPairInitializer<FName, int64>(Kvp.Key.CreateName(UHTMakefile), Kvp.Value));
 	}
 
 	Enum->SetEnums(Names, (UEnum::ECppForm)CppForm);
@@ -57,7 +57,7 @@ FArchive& operator<<(FArchive& Ar, FEnumArchiveProxy& EnumArchiveProxy)
 	Ar << static_cast<FFieldArchiveProxy&>(EnumArchiveProxy);
 	Ar << EnumArchiveProxy.EnumFlags;
 	Ar << EnumArchiveProxy.CppForm;
-	Ar << EnumArchiveProxy.EnumValues;
+	Ar << EnumArchiveProxy.EnumValuesX;
 	Ar << EnumArchiveProxy.CppType;
 
 	return Ar;
diff --git a/Engine/Source/Programs/UnrealHeaderTool/Private/UHTMakefile/EnumArchiveProxy.h b/Engine/Source/Programs/UnrealHeaderTool/Private/UHTMakefile/EnumArchiveProxy.h
--- a/Engine/Source/Programs/UnrealHeaderTool/Private/UHTMakefile/EnumArchiveProxy.h
+++ b/Engine/Source/Programs/UnrealHeaderTool/Private/UHTMakefile/EnumArchiveProxy.h
@@ -18,5 +18,5 @@ struct FEnumArchiveProxy : public FFieldArchiveProxy
 	uint32 EnumFlags;
 	uint32 CppForm;
 	FString CppType;
-	TArray<TPair<FNameArchiveProxy, uint8>> EnumValues;
+	TArray<TPair<FNameArchiveProxy, int64>> EnumValuesX;
 };
diff --git a/Engine/Source/Programs/UnrealHeaderTool/Private/UHTMakefile/MakefileHelpers.h b/Engine/Source/Programs/UnrealHeaderTool/Private/UHTMakefile/MakefileHelpers.h
--- a/Engine/Source/Programs/UnrealHeaderTool/Private/UHTMakefile/MakefileHelpers.h
+++ b/Engine/Source/Programs/UnrealHeaderTool/Private/UHTMakefile/MakefileHelpers.h
@@ -151,6 +151,7 @@ enum class ESerializedObjectType : uint32
 	EEnumUnderlyingType,
 	EStructNameMapEntry,
 	EInterfaceAllocation,
+	EEnumProperty,
 };
 
 enum class EUHTMakefileLoadingPhase : uint8
diff --git a/Engine/Source/Programs/UnrealHeaderTool/Private/UHTMakefile/PropertyArchiveProxy.cpp b/Engine/Source/Programs/UnrealHeaderTool/Private/UHTMakefile/PropertyArchiveProxy.cpp
--- a/Engine/Source/Programs/UnrealHeaderTool/Private/UHTMakefile/PropertyArchiveProxy.cpp
+++ b/Engine/Source/Programs/UnrealHeaderTool/Private/UHTMakefile/PropertyArchiveProxy.cpp
@@ -96,6 +96,37 @@ FArchive& operator<<(FArchive& Ar, FSetPropertyArchiveProxy& SetPropertyArchiveP
 	return Ar;
 }
 
+FEnumPropertyArchiveProxy::FEnumPropertyArchiveProxy(FUHTMakefile& UHTMakefile, const UEnumProperty* EnumProperty)
+	: FPropertyArchiveProxy(UHTMakefile, EnumProperty)
+{
+	EnumIndex = UHTMakefile.GetEnumIndex(EnumProperty->Enum);
+	UnderlyingPropertyIndex = UHTMakefile.GetPropertyIndex(EnumProperty->UnderlyingProp);
+}
+
+UEnumProperty* FEnumPropertyArchiveProxy::CreateEnumProperty(const FUHTMakefile& UHTMakefile) const
+{
+	UObject* Outer = UHTMakefile.GetObjectByIndex(OuterIndex);
+	UEnumProperty* EnumProperty = new (EC_InternalUseOnlyConstructor, Outer, Name.CreateName(UHTMakefile), (EObjectFlags)ObjectFlagsUint32) UEnumProperty(FObjectInitializer());
+	PostConstruct(EnumProperty, UHTMakefile);
+	return EnumProperty;
+}
+
+void FEnumPropertyArchiveProxy::Resolve(UEnumProperty* EnumProperty, const FUHTMakefile& UHTMakefile) const
+{
+	FPropertyArchiveProxy::Resolve(EnumProperty, UHTMakefile);
+	EnumProperty->Enum = UHTMakefile.GetEnumByIndex(EnumIndex);
+	EnumProperty->UnderlyingProp = CastChecked<UNumericProperty>(UHTMakefile.GetPropertyByIndex(UnderlyingPropertyIndex));
+}
+
+FArchive& operator<<(FArchive& Ar, FEnumPropertyArchiveProxy& EnumPropertyArchiveProxy)
+{
+	Ar << static_cast<FPropertyArchiveProxy&>(EnumPropertyArchiveProxy);
+	Ar << EnumPropertyArchiveProxy.EnumIndex;
+	Ar << EnumPropertyArchiveProxy.UnderlyingPropertyIndex;
+
+	return Ar;
+}
+
 FBytePropertyArchiveProxy::FBytePropertyArchiveProxy(FUHTMakefile& UHTMakefile, const UByteProperty* ByteProperty)
 	: FPropertyArchiveProxy(UHTMakefile, ByteProperty)
 {
diff --git a/Engine/Source/Programs/UnrealHeaderTool/Private/UHTMakefile/PropertyArchiveProxy.h b/Engine/Source/Programs/UnrealHeaderTool/Private/UHTMakefile/PropertyArchiveProxy.h
--- a/Engine/Source/Programs/UnrealHeaderTool/Private/UHTMakefile/PropertyArchiveProxy.h
+++ b/Engine/Source/Programs/UnrealHeaderTool/Private/UHTMakefile/PropertyArchiveProxy.h
@@ -29,6 +29,25 @@ struct FPropertyArchiveProxy : public FFieldArchiveProxy
 	int32 PostConstructLinkNextIndex;
 };
 
+struct FEnumPropertyArchiveProxy : public FPropertyArchiveProxy
+{
+	FEnumPropertyArchiveProxy() { }
+	FEnumPropertyArchiveProxy(FUHTMakefile& UHTMakefile, const UEnumProperty* EnumProperty);
+
+	UEnumProperty* CreateEnumProperty(const FUHTMakefile& UHTMakefile) const;
+
+	void Resolve(UEnumProperty* EnumProperty, const FUHTMakefile& UHTMakefile) const;
+	static void AddReferencedNames(const UEnumProperty* Property, FUHTMakefile& UHTMakefile)
+	{
+		FPropertyArchiveProxy::AddReferencedNames(Property, UHTMakefile);
+	}
+
+	friend FArchive& operator<<(FArchive& Ar, FEnumPropertyArchiveProxy& BytePropertyArchiveProxy);
+
+	int32 EnumIndex;
+	int32 UnderlyingPropertyIndex;
+};
+
 struct FBytePropertyArchiveProxy : public FPropertyArchiveProxy
 {
 	FBytePropertyArchiveProxy() { }
diff --git a/Engine/Source/Programs/UnrealHeaderTool/Private/UHTMakefile/TokenArchiveProxy.cpp b/Engine/Source/Programs/UnrealHeaderTool/Private/UHTMakefile/TokenArchiveProxy.cpp
--- a/Engine/Source/Programs/UnrealHeaderTool/Private/UHTMakefile/TokenArchiveProxy.cpp
+++ b/Engine/Source/Programs/UnrealHeaderTool/Private/UHTMakefile/TokenArchiveProxy.cpp
@@ -23,6 +23,9 @@ FTokenArchiveProxy::FTokenArchiveProxy(const FUHTMakefile& UHTMakefile, const FT
 	case CPT_Int:
 		Int = Token->Int;
 		break;
+	case CPT_Int64:
+		Int64 = Token->Int64;
+		break;
 	case CPT_Bool:
 		NativeBool = Token->NativeBool;
 		break;
@@ -74,6 +77,9 @@ void FTokenArchiveProxy::PostConstruct(FToken* Token) const
 	case CPT_Int:
 		Token->Int = Int;
 		break;
+	case CPT_Int64:
+		Token->Int64 = Int64;
+		break;
 	case CPT_Bool:
 		Token->NativeBool = NativeBool;
 		break;
@@ -119,6 +125,9 @@ FArchive& operator<<(FArchive& Ar, FTokenArchiveProxy& TokenArchiveProxy)
 	case CPT_Int:
 		Ar << TokenArchiveProxy.Int;
 		break;
+	case CPT_Int64:
+		Ar << TokenArchiveProxy.Int64;
+		break;
 	case CPT_Bool:
 		Ar << TokenArchiveProxy.NativeBool;
 		break;
diff --git a/Engine/Source/Programs/UnrealHeaderTool/Private/UHTMakefile/TokenArchiveProxy.h b/Engine/Source/Programs/UnrealHeaderTool/Private/UHTMakefile/TokenArchiveProxy.h
--- a/Engine/Source/Programs/UnrealHeaderTool/Private/UHTMakefile/TokenArchiveProxy.h
+++ b/Engine/Source/Programs/UnrealHeaderTool/Private/UHTMakefile/TokenArchiveProxy.h
@@ -31,6 +31,7 @@ struct FTokenArchiveProxy : public FPropertyBaseArchiveProxy
 	{
 		uint8 Byte;								// If CPT_Byte.
 		int32 Int;								// If CPT_Int.
+		int32 Int64;							// If CPT_Int64.
 		bool NativeBool;						// if CPT_Bool
 		float Float;							// If CPT_Float.
 		double Double;							// If CPT_Double.
diff --git a/Engine/Source/Programs/UnrealHeaderTool/Private/UHTMakefile/UHTMakefile.cpp b/Engine/Source/Programs/UnrealHeaderTool/Private/UHTMakefile/UHTMakefile.cpp
--- a/Engine/Source/Programs/UnrealHeaderTool/Private/UHTMakefile/UHTMakefile.cpp
+++ b/Engine/Source/Programs/UnrealHeaderTool/Private/UHTMakefile/UHTMakefile.cpp
@@ -18,6 +18,35 @@
 #include "ScopeExit.h"
 
 
+namespace UE4Makefile_Private
+{
+	enum class EVersion : uint32
+	{
+		FirstExplicitVersion // there was a version before this, but as there was no header, it was indeterminate
+	};
+
+	static const uint32 HeaderFourCC = 0x5548544D; // UHTM
+
+	struct FHeader
+	{
+		FHeader()
+			: FourCC(HeaderFourCC)
+			, Version(EVersion::FirstExplicitVersion)
+		{
+		}
+
+		friend FArchive& operator<<(FArchive& Ar, FHeader& Header)
+		{
+			Ar << Header.FourCC;
+			Ar << Header.Version;
+			return Ar;
+		}
+
+		uint32 FourCC;
+		EVersion Version;
+	};
+}
+
 void FUHTMakefile::SetupScopeArrays()
 {
 	TSet<FScope*> AllScopesSet;
@@ -81,6 +110,12 @@ void FUHTMakefile::SetupProxies()
 		PropertyProxies.Add(FPropertyArchiveProxy(*this, Property));
 	}
 
+	EnumPropertyProxies.Empty(EnumProperties.Num());
+	for (UEnumProperty* EnumProperty : EnumProperties)
+	{
+		EnumPropertyProxies.Add(FEnumPropertyArchiveProxy(*this, EnumProperty));
+	}
+
 	BytePropertyProxies.Empty(ByteProperties.Num());
 	for (UByteProperty* ByteProperty : ByteProperties)
 	{
@@ -344,9 +379,9 @@ void FUHTMakefile::SetupProxies()
 	}
 
 	ClassMetaDataArchiveProxies.Empty(ClassMetaDatas.Num());
-	for (const TScopedPointer<FClassMetaData>& MetaData : ClassMetaDatas)
+	for (const TUniquePtr<FClassMetaData>& MetaData : ClassMetaDatas)
 	{
-		ClassMetaDataArchiveProxies.Add(FClassMetaDataArchiveProxy(*this, MetaData.GetOwnedPointer()));
+		ClassMetaDataArchiveProxies.Add(FClassMetaDataArchiveProxy(*this, MetaData.Get()));
 	}
 
 	TokenProxies.Empty(Tokens.Num());
@@ -364,7 +399,7 @@ void FUHTMakefile::SetupProxies()
 	EnumUnderlyingTypeProxies.Empty(EnumUnderlyingTypes.Num());
 	for (auto& Kvp : EnumUnderlyingTypes)
 	{
-		EnumUnderlyingTypeProxies.Add(TPairInitializer<int32, EPropertyType>(GetEnumIndex(Kvp.Key), Kvp.Value));
+		EnumUnderlyingTypeProxies.Add(TPairInitializer<int32, EUnderlyingEnumType>(GetEnumIndex(Kvp.Key), Kvp.Value));
 	}
 
 	StructNameMapEntryProxies.Empty(StructNameMapEntries.Num());
@@ -383,6 +418,7 @@ void FUHTMakefile::SetupProxies()
 int32 FUHTMakefile::GetPropertyCount() const
 {
 	return Properties.Num()
+		+ EnumProperties.Num()
 		+ ByteProperties.Num()
 		+ Int8Properties.Num()
 		+ Int16Properties.Num()
@@ -416,6 +452,7 @@ int32 FUHTMakefile::GetPropertyCount() const
 int32 FUHTMakefile::GetPropertyProxiesCount() const
 {
 	return Properties.Num()
+		+ EnumPropertyProxies.Num()
 		+ BytePropertyProxies.Num()
 		+ Int8PropertyProxies.Num()
 		+ Int16PropertyProxies.Num()
@@ -520,17 +557,19 @@ void FUHTMakefile::SaveToFile(const TCHAR* MakefilePath)
 	double UHTMakefileLoadTime = 0.0;
 	FDurationTimer UHTMakefileLoadTimer(UHTMakefileLoadTime);
 	UHTMakefileLoadTimer.Start();
+	UE4Makefile_Private::FHeader Header;
+	*UHTMakefileArchive << Header;
 	*UHTMakefileArchive << *this;
 	UHTMakefileLoadTimer.Stop();
 	UE_LOG(LogCompile, Log, TEXT("Saving UHT makefile took %f seconds"), UHTMakefileLoadTime);
 }
 
-void FUHTMakefile::LoadFromFile(const TCHAR* MakefilePath, FManifest* InManifest)
+bool FUHTMakefile::LoadFromFile(const TCHAR* MakefilePath, FManifest* InManifest)
 {
 	FArchive* UHTMakefileArchive = IFileManager::Get().CreateFileReader(MakefilePath);
 	if (!UHTMakefileArchive)
 	{
-		return;
+		return false;
 	}
 
 	ON_SCOPE_EXIT
@@ -541,11 +580,19 @@ void FUHTMakefile::LoadFromFile(const TCHAR* MakefilePath, FManifest* InManifest
 	double UHTMakefileLoadTime = 0.0;
 	FDurationTimer UHTMakefileLoadTimer(UHTMakefileLoadTime);
 	UHTMakefileLoadTimer.Start();
+	UE4Makefile_Private::FHeader Header;
+	*UHTMakefileArchive << Header;
+	if (UHTMakefileArchive->IsError() || Header.FourCC != UE4Makefile_Private::HeaderFourCC)
+	{
+		return false;
+	}
 	*UHTMakefileArchive << *this;
 	UHTMakefileLoadTimer.Stop();
 	UE_LOG(LogCompile, Log, TEXT("Loading UHT makefile took %f seconds"), UHTMakefileLoadTime);
 
 	UpdateModulesCompatibility(InManifest);
+
+	return !UHTMakefileArchive->IsError();
 }
 
 bool FUHTMakefile::CanLoadModule(const FManifestModule& ManifestModule)
@@ -826,6 +873,9 @@ void FUHTMakefile::CreateObject(ESerializedObjectType ObjectType, int32 Index)
 	case ESerializedObjectType::EInterfaceAllocation:
 		CreateInterfaceAllocation(Index);
 		break;
+	case ESerializedObjectType::EEnumProperty:
+		CreateEnumProperty(Index);
+		break;
 	default:
 		check(0);
 		break;
@@ -998,6 +1048,9 @@ void FUHTMakefile::ResolveObject(ESerializedObjectType ObjectType, int32 Index)
 	case ESerializedObjectType::EInterfaceAllocation:
 		// Intentionally empty.
 		break;
+	case ESerializedObjectType::EEnumProperty:
+		ResolveEnumProperty(Index);
+		break;
 	default:
 		check(0);
 		break;
@@ -1364,6 +1417,13 @@ int32 FUHTMakefile::GetPropertyIndex(const UProperty* Property) const
 	}
 	Result += Properties.Num();
 
+	CurrentIndex = EnumPropertyIndexes.Find(Cast<UEnumProperty>(Property));
+	if (CurrentIndex)
+	{
+		return Result + *CurrentIndex;
+	}
+	Result += EnumProperties.Num();
+
 	CurrentIndex = BytePropertyIndexes.Find(Cast<UByteProperty>(Property));
 	if (CurrentIndex)
 	{
@@ -1576,6 +1636,12 @@ UProperty* FUHTMakefile::GetPropertyByIndex(int32 Index) const
 	}
 	Index -= PropertyProxies.Num();
 
+	if (EnumProperties.IsValidIndex(Index) && EnumPropertyProxies.IsValidIndex(Index))
+	{
+		return EnumProperties[Index];
+	}
+	Index -= EnumPropertyProxies.Num();
+
 	if (ByteProperties.IsValidIndex(Index) && BytePropertyProxies.IsValidIndex(Index))
 	{
 		return ByteProperties[Index];
@@ -2058,6 +2124,11 @@ void FUHTMakefile::SetupAndSaveReferencedNames(FArchive& Ar)
 		FPropertyArchiveProxy::AddReferencedNames(Property, *this);
 	}
 
+	for (UEnumProperty* EnumProperty : EnumProperties)
+	{
+		FObjectBaseArchiveProxy::AddReferencedNames(EnumProperty, *this);
+	}
+
 	for (UByteProperty* ByteProperty : ByteProperties)
 	{
 		FObjectBaseArchiveProxy::AddReferencedNames(ByteProperty, *this);
@@ -2220,9 +2291,9 @@ void FUHTMakefile::SetupAndSaveReferencedNames(FArchive& Ar)
 		FMultipleInheritanceBaseClassArchiveProxy::AddReferencedNames(MultipleInheritanceBaseClass, *this); 
 	}
 
-	for (FClassMetaData* ClassMetaData : ClassMetaDatas)
+	for (const TUniquePtr<FClassMetaData>& ClassMetaData : ClassMetaDatas)
 	{
-		FClassMetaDataArchiveProxy::AddReferencedNames(ClassMetaData, *this);
+		FClassMetaDataArchiveProxy::AddReferencedNames(ClassMetaData.Get(), *this);
 	}
 
 	for (FToken* Token : Tokens)
@@ -2296,6 +2367,7 @@ FArchive& operator<<(FArchive& Ar, FUHTMakefile& UHTMakefile)
 	Ar << UHTMakefile.UnrealTypeDefinitionInfoProxies;
 	Ar << UHTMakefile.TypeDefinitionInfoMapArchiveProxy;
 	Ar << UHTMakefile.PropertyProxies;
+	Ar << UHTMakefile.EnumPropertyProxies;
 	Ar << UHTMakefile.BytePropertyProxies;
 	Ar << UHTMakefile.Int8PropertyProxies;
 	Ar << UHTMakefile.Int16PropertyProxies;
diff --git a/Engine/Source/Programs/UnrealHeaderTool/Private/UHTMakefile/UHTMakefile.h b/Engine/Source/Programs/UnrealHeaderTool/Private/UHTMakefile/UHTMakefile.h
--- a/Engine/Source/Programs/UnrealHeaderTool/Private/UHTMakefile/UHTMakefile.h
+++ b/Engine/Source/Programs/UnrealHeaderTool/Private/UHTMakefile/UHTMakefile.h
@@ -16,6 +16,8 @@
 #include "UHTMakefile/TypeDefinitionInfoMapArchiveProxy.h"
 #include "UHTMakefile/PropertyBaseArchiveProxy.h"
 #include "UHTMakefile/CompilerMetadataManagerArchiveProxy.h"
+#include "UnderlyingEnumType.h"
+#include "UniquePtr.h"
 
 /**
  * UHT makefiles overview.
@@ -343,6 +345,24 @@ public:
 	int32 GetPropertyIndex(const UProperty* Property) const;
 	UProperty* GetPropertyByIndex(int32 Index) const;
 
+	void AddEnumProperty(FUnrealSourceFile* UnrealSourceFile, UEnumProperty* EnumProperty)
+	{
+		int32 Index = EnumProperties.Add(EnumProperty);
+		GetHeaderDescriptor(UnrealSourceFile).AddEntry(AddObject(ESerializedObjectType::EEnumProperty, Index));
+		EnumPropertyIndexes.Add(EnumProperty, Index);
+	}
+	void CreateEnumProperty(int32 Index)
+	{
+		checkSlow(Index == EnumProperties.Num());
+		UEnumProperty* EnumProperty = EnumPropertyProxies[Index].CreateEnumProperty(*this);
+		Index = EnumProperties.Add(EnumProperty);
+		EnumPropertyIndexes.Add(EnumProperty, Index);
+	}
+	void ResolveEnumProperty(int32 Index)
+	{
+		EnumPropertyProxies[Index].Resolve(EnumProperties[Index], *this);
+	}
+
 	void AddByteProperty(FUnrealSourceFile* UnrealSourceFile, UByteProperty* ByteProperty)
 	{
 		int32 Index = ByteProperties.Add(ByteProperty);
@@ -979,14 +999,14 @@ public:
 		GetHeaderDescriptor(UnrealSourceFile).AddEntry(AddObject(ESerializedObjectType::EPublicClassSetEntry, Index));
 	}
 
-	void AddGEnumUnderlyingType(FUnrealSourceFile* UnrealSourceFile, UEnum* Enum, EPropertyType Type)
+	void AddGEnumUnderlyingType(FUnrealSourceFile* UnrealSourceFile, UEnum* Enum, EUnderlyingEnumType Type)
 	{
-		int32 Index = EnumUnderlyingTypes.Add(TPairInitializer<UEnum*, EPropertyType>(Enum, Type));
+		int32 Index = EnumUnderlyingTypes.Add(TPairInitializer<UEnum*, EUnderlyingEnumType>(Enum, Type));
 		GetHeaderDescriptor(UnrealSourceFile).AddEntry(AddObject(ESerializedObjectType::EEnumUnderlyingType, Index));
 	}
 	void CreateGEnumUnderlyingType(int32 Index)
 	{
-		EnumUnderlyingTypes.Add(TPairInitializer<UEnum*, EPropertyType>(nullptr, EnumUnderlyingTypeProxies[Index].Value));
+		EnumUnderlyingTypes.Add(TPairInitializer<UEnum*, EUnderlyingEnumType>(nullptr, EnumUnderlyingTypeProxies[Index].Value));
 	}
 	void ResolveGEnumUnderlyingType(int32 Index)
 	{
@@ -1083,9 +1103,9 @@ public:
 		return GeneratedCodeCRCProxies[Index];
 	}
 	
-	void AddClassMetaData(FUnrealSourceFile* SourceFile, TScopedPointer<FClassMetaData>& ClassMetaData)
+	void AddClassMetaData(FUnrealSourceFile* SourceFile, TUniquePtr<FClassMetaData> ClassMetaData)
 	{
-		int32 Index = ClassMetaDatas.Add(ClassMetaData);
+		int32 Index = ClassMetaDatas.Add(MoveTemp(ClassMetaData));
 		GetHeaderDescriptor(SourceFile).AddEntry(AddObject(ESerializedObjectType::EClassMetaData, Index));
 	}
 	void CreateClassMetaData(int32 Index)
@@ -1095,7 +1115,7 @@ public:
 	}
 	void ResolveClassMetaData(int32 Index)
 	{
-		ClassMetaDataArchiveProxies[Index].Resolve(ClassMetaDatas[Index], *this);
+		ClassMetaDataArchiveProxies[Index].Resolve(ClassMetaDatas[Index].Get(), *this);
 	}
 
 	void AddPropertyDataEntry(FUnrealSourceFile* UnrealSourceFile, TSharedPtr<FTokenData>& TokenData, UProperty* Property)
@@ -1311,7 +1331,7 @@ public:
 		Manifest = InManifest;
 	}
 	bool CanLoadModule(const FManifestModule& ManifestModule);
-	void LoadFromFile(const TCHAR* MakefilePath, FManifest* Manifest);
+	bool LoadFromFile(const TCHAR* MakefilePath, FManifest* Manifest);
 	void SaveToFile(const TCHAR* MakefilePath);
 private:
 	bool bShouldForceRegeneration;
@@ -1396,8 +1416,8 @@ private:
 
 	EUHTMakefileLoadingPhase LoadingPhase;
 
-	TArray<TPair<UEnum*, EPropertyType>> EnumUnderlyingTypes;
-	TArray<TPair<int32, EPropertyType>> EnumUnderlyingTypeProxies;
+	TArray<TPair<UEnum*, EUnderlyingEnumType>> EnumUnderlyingTypes;
+	TArray<TPair<int32, EUnderlyingEnumType>> EnumUnderlyingTypeProxies;
 
 	TArray<FPropertyData*> PropertyDatas;
 	TArray<FPropertyDataArchiveProxy> PropertyDataProxies;
@@ -1405,7 +1425,7 @@ private:
 	TArray<FNameArchiveProxy> ModuleNamesProxy;
 	TArray<FName> ModuleNames;
 
-	TArray<TScopedPointer<FClassMetaData>> ClassMetaDatas;
+	TArray<TUniquePtr<FClassMetaData>> ClassMetaDatas;
 	TArray<FClassMetaDataArchiveProxy> ClassMetaDataArchiveProxies;
 
 	TArray<FPackageArchiveProxy> PackageProxies;
@@ -1479,6 +1499,10 @@ private:
 	TArray<UProperty*> Properties;
 	TMap<UProperty*, int32> PropertyIndexes;
 
+	TArray<FEnumPropertyArchiveProxy> EnumPropertyProxies;
+	TArray<UEnumProperty*> EnumProperties;
+	TMap<UEnumProperty*, int32> EnumPropertyIndexes;
+
 	TArray<FBytePropertyArchiveProxy> BytePropertyProxies;
 	TArray<UByteProperty*> ByteProperties;
 	TMap<UByteProperty*, int32> BytePropertyIndexes;
diff --git a/Engine/Source/Programs/UnrealHeaderTool/Private/UnderlyingEnumType.h b/Engine/Source/Programs/UnrealHeaderTool/Private/UnderlyingEnumType.h
new file mode 100644
--- /dev/null
+++ b/Engine/Source/Programs/UnrealHeaderTool/Private/UnderlyingEnumType.h
@@ -0,0 +1,16 @@
+// Copyright 1998-2016 Epic Games, Inc. All Rights Reserved.
+
+#pragma once
+
+enum class EUnderlyingEnumType
+{
+	Unspecified,
+	uint8,
+	uint16,
+	uint32,
+	uint64,
+	int8,
+	int16,
+	int32,
+	int64
+};
diff --git a/Engine/Source/Programs/UnrealLightmass/Private/LightmassCore/LMCore.cpp b/Engine/Source/Programs/UnrealLightmass/Private/LightmassCore/LMCore.cpp
--- a/Engine/Source/Programs/UnrealLightmass/Private/LightmassCore/LMCore.cpp
+++ b/Engine/Source/Programs/UnrealLightmass/Private/LightmassCore/LMCore.cpp
@@ -2,6 +2,7 @@
 
 #include "LightmassPCH.h"
 #include "LMCore.h"
+#include "UniquePtr.h"
 
 namespace Lightmass
 {
@@ -92,7 +93,7 @@ void FLightmassLog::Flush()
 	File->Flush();
 }
 
-static TScopedPointer<FLightmassLog>	GScopedLog;
+static TUniquePtr<FLightmassLog>	GScopedLog;
 
 FLightmassLog* FLightmassLog::Get()
 {
diff --git a/Engine/Source/Programs/UnrealPak/Private/SignedArchiveWriter.cpp b/Engine/Source/Programs/UnrealPak/Private/SignedArchiveWriter.cpp
--- a/Engine/Source/Programs/UnrealPak/Private/SignedArchiveWriter.cpp
+++ b/Engine/Source/Programs/UnrealPak/Private/SignedArchiveWriter.cpp
@@ -29,29 +29,13 @@ FSignedArchiveWriter::~FSignedArchiveWriter()
 
 void FSignedArchiveWriter::SerializeBufferAndSign()
 {
-	// Hash the buffer
-	FDecryptedSignature SourceSignature;
-	
 	// Compute a hash for this buffer data
-	SourceSignature.Data = FCrc::MemCrc32(&Buffer[0], Buffer.Num());
-
-	// Encrypt the signature
-	FEncryptedSignature EncryptedSignature;
-	FEncryption::EncryptSignature(SourceSignature, EncryptedSignature, PrivateKey);
+	ChunkHashes.Add(ComputePakChunkHash(&Buffer[0], Buffer.Num()));
 
 	// Flush the buffer
 	PakWriter.Serialize(&Buffer[0], Buffer.Num());
 	BufferArchive.Seek(0);
 	Buffer.Empty(FPakInfo::MaxChunkDataSize);
-
-	// Collect the signature so we can write it out at the end
-	ChunkSignatures.Add(EncryptedSignature);
-
-#if 0
-	FDecryptedSignature<GPakFileChunkHashSize> TestSignature;
-	FEncryption::DecryptBytes(TestSignature.Data, Signature.Data, GPakFileChunkHashSize, PublicKey);
-	check(FMemory::Memcmp(TestSignature.Data, Hash, sizeof(Hash)) == 0);
-#endif
 }
 
 bool FSignedArchiveWriter::Close()
@@ -61,8 +45,14 @@ bool FSignedArchiveWriter::Close()
 		SerializeBufferAndSign();
 	}
 
+	FEncryptedSignature EncryptedMasterHash;
+	FDecryptedSignature DecryptedMasterHash;
+	DecryptedMasterHash.Data = ComputePakChunkHash((const uint8*)&ChunkHashes[0], ChunkHashes.Num() * sizeof(TPakChunkHash));
+	FEncryption::EncryptSignature(DecryptedMasterHash, EncryptedMasterHash, PrivateKey);
+
 	FArchive* SignatureWriter = IFileManager::Get().CreateFileWriter(*PakSignaturesFilename);
-	*SignatureWriter << ChunkSignatures;
+	*SignatureWriter << EncryptedMasterHash;
+	*SignatureWriter << ChunkHashes;
 	delete SignatureWriter;
 
 	return FArchive::Close();
diff --git a/Engine/Source/Programs/UnrealPak/Private/SignedArchiveWriter.h b/Engine/Source/Programs/UnrealPak/Private/SignedArchiveWriter.h
--- a/Engine/Source/Programs/UnrealPak/Private/SignedArchiveWriter.h
+++ b/Engine/Source/Programs/UnrealPak/Private/SignedArchiveWriter.h
@@ -22,8 +22,8 @@ class FSignedArchiveWriter : public FArchive
 	FEncryptionKey PublicKey;
 	/** Encryption key */
 	FEncryptionKey PrivateKey;
-	/** Signatures */
-	TArray<FEncryptedSignature> ChunkSignatures;
+	/** Hashes */
+	TArray<TPakChunkHash> ChunkHashes;
 
 	/** 
 	 * Serializes and signs a buffer
diff --git a/Engine/Source/Programs/UnrealPak/Private/UnrealPak.cpp b/Engine/Source/Programs/UnrealPak/Private/UnrealPak.cpp
--- a/Engine/Source/Programs/UnrealPak/Private/UnrealPak.cpp
+++ b/Engine/Source/Programs/UnrealPak/Private/UnrealPak.cpp
@@ -174,14 +174,14 @@ FString GetCommonRootPath(TArray<FPakInputPair>& FilesToAdd)
 
 bool FCompressedFileBuffer::CompressFileToWorkingBuffer(const FPakInputPair& InFile, uint8*& InOutPersistentBuffer, int64& InOutBufferSize, ECompressionFlags CompressionMethod, const int32 CompressionBlockSize, const int32 CompressionBitWindow)
 {
-	TAutoPtr<FArchive> FileHandle(IFileManager::Get().CreateFileReader(*InFile.Source));
-	if(!FileHandle.IsValid())
+	TUniquePtr<FArchive> FileHandle(IFileManager::Get().CreateFileReader(*InFile.Source));
+	if(!FileHandle)
 	{
 		TotalCompressedSize = 0;
 		return false;
 	}
 
-	Reinitialize(FileHandle.GetOwnedPointer(), CompressionMethod, CompressionBlockSize);
+	Reinitialize(FileHandle.Get(), CompressionMethod, CompressionBlockSize);
 	const int64 FileSize = OriginalSize;
 	const int64 PaddedEncryptedFileSize = Align(FileSize,FAES::AESBlockSize);
 	if(InOutBufferSize < PaddedEncryptedFileSize)
@@ -227,8 +227,9 @@ bool FCompressedFileBuffer::CompressFileToWorkingBuffer(const FPakInputPair& InF
 			int32 EncryptionBlockPadding = Align(TotalCompressedSize,FAES::AESBlockSize);
 			for(int64 FillIndex=TotalCompressedSize; FillIndex < EncryptionBlockPadding; ++FillIndex)
 			{
-				// Fill the trailing buffer with random bytes from file
-				CompressedBuffer.Get()[FillIndex] = CompressedBuffer.Get()[rand()%TotalCompressedSize];
+				// Fill the trailing buffer with bytes from file. Note that this is now from a fixed location
+				// rather than a random one so that we produce deterministic results
+				CompressedBuffer.Get()[FillIndex] = CompressedBuffer.Get()[FillIndex % TotalCompressedSize];
 			}
 			TotalCompressedSize += EncryptionBlockPadding - TotalCompressedSize;
 		}
@@ -237,9 +238,9 @@ bool FCompressedFileBuffer::CompressFileToWorkingBuffer(const FPakInputPair& InF
 	return true;
 }
 
-bool PrepareCopyFileToPak(const FString& InMountPoint, const FPakInputPair& InFile, uint8*& InOutPersistentBuffer, int64& InOutBufferSize, FPakEntryPair& OutNewEntry, uint8*& OutDataToWrite, int64& OutSizeToWrite)
+bool PrepareCopyFileToPak(const FString& InMountPoint, const FPakInputPair& InFile, uint8*& InOutPersistentBuffer, int64& InOutBufferSize, FPakEntryPair& OutNewEntry, uint8*& OutDataToWrite, int64& OutSizeToWrite, ANSICHAR* InEncryptionKey)
 {	
-	TAutoPtr<FArchive> FileHandle(IFileManager::Get().CreateFileReader(*InFile.Source));
+	TUniquePtr<FArchive> FileHandle(IFileManager::Get().CreateFileReader(*InFile.Source));
 	bool bFileExists = FileHandle.IsValid();
 	if (bFileExists)
 	{
@@ -263,7 +264,7 @@ bool PrepareCopyFileToPak(const FString& InMountPoint, const FPakInputPair& InFi
 
 		{
 			OutSizeToWrite = FileSize;
-			if (InFile.bNeedEncryption)
+			if (InFile.bNeedEncryption && InEncryptionKey)
 			{
 				for(int64 FillIndex=FileSize; FillIndex < PaddedEncryptedFileSize && InFile.bNeedEncryption; ++FillIndex)
 				{
@@ -272,7 +273,7 @@ bool PrepareCopyFileToPak(const FString& InMountPoint, const FPakInputPair& InFi
 				}
 
 				//Encrypt the buffer before writing it to disk
-				FAES::EncryptData(InOutPersistentBuffer, PaddedEncryptedFileSize);
+				FAES::EncryptData(InOutPersistentBuffer, PaddedEncryptedFileSize, InEncryptionKey);
 				// Update the size to be written
 				OutSizeToWrite = PaddedEncryptedFileSize;
 				OutNewEntry.Info.bEncrypted = true;
@@ -302,7 +303,7 @@ void FinalizeCopyCompressedFileToPak(FArchive& InPak, const FCompressedFileBuffe
 	}
 }
 
-bool PrepareCopyCompressedFileToPak(const FString& InMountPoint, const FPakInputPair& InFile, const FCompressedFileBuffer& CompressedFile, FPakEntryPair& OutNewEntry, uint8*& OutDataToWrite, int64& OutSizeToWrite)
+bool PrepareCopyCompressedFileToPak(const FString& InMountPoint, const FPakInputPair& InFile, const FCompressedFileBuffer& CompressedFile, FPakEntryPair& OutNewEntry, uint8*& OutDataToWrite, int64& OutSizeToWrite, ANSICHAR* InEncryptionKey)
 {
 	if (CompressedFile.TotalCompressedSize == 0)
 	{
@@ -312,9 +313,9 @@ bool PrepareCopyCompressedFileToPak(const FString& InMountPoint, const FPakInput
 	OutNewEntry.Info.CompressionMethod = CompressedFile.FileCompressionMethod;
 	OutNewEntry.Info.CompressionBlocks.AddZeroed(CompressedFile.CompressedBlocks.Num());
 
-	if (InFile.bNeedEncryption)
+	if (InFile.bNeedEncryption && InEncryptionKey)
 	{
-		FAES::EncryptData(CompressedFile.CompressedBuffer.Get(), CompressedFile.TotalCompressedSize);
+		FAES::EncryptData(CompressedFile.CompressedBuffer.Get(), CompressedFile.TotalCompressedSize, InEncryptionKey);
 	}
 
 	//Hash the final buffer thats written
@@ -709,51 +710,133 @@ bool UncompressCopyFile(FArchive& Dest, FArchive& Source, const FPakEntry& Entry
 /**
  * Creates a pak file writer. This can be a signed writer if the encryption keys are specified in the command line
  */
-FArchive* CreatePakWriter(const TCHAR* Filename)
+FArchive* CreatePakWriter(const TCHAR* Filename, FKeyPair& OutEncryptionKeyPair, FString& OutAESKey)
 {
+	OutEncryptionKeyPair.PrivateKey.Exponent.Zero();
+	OutEncryptionKeyPair.PrivateKey.Modulus.Zero();
+	OutEncryptionKeyPair.PublicKey.Exponent.Zero();
+	OutEncryptionKeyPair.PublicKey.Modulus.Zero();
+	OutAESKey = TEXT("");
+
 	FArchive* Writer = IFileManager::Get().CreateFileWriter(Filename);
 	FString KeyFilename;
-	if (Writer && FParse::Value(FCommandLine::Get(), TEXT("sign="), KeyFilename, false))
+	bool bSigningEnabled = false;
+	
+	if (Writer)
 	{
-		FKeyPair Pair;
-		if (KeyFilename.StartsWith(TEXT("0x")))
+		if (FParse::Param(FCommandLine::Get(), TEXT("encryptionini")))
 		{
-			TArray<FString> KeyValueText;
-			if (KeyFilename.ParseIntoArray(KeyValueText, TEXT("+"), true) == 3)
-			{
-				Pair.PrivateKey.Exponent.Parse(KeyValueText[0]);
-				Pair.PrivateKey.Modulus.Parse(KeyValueText[1]);
-				Pair.PublicKey.Exponent.Parse(KeyValueText[2]);
-				Pair.PublicKey.Modulus = Pair.PrivateKey.Modulus;
-				UE_LOG(LogPakFile, Display, TEXT("Parsed signature keys from command line."));
-			}
-			else
+			FString ProjectDir, EngineDir, Platform;
+
+			if (FParse::Value(FCommandLine::Get(), TEXT("projectdir="), ProjectDir, false)
+				&& FParse::Value(FCommandLine::Get(), TEXT("enginedir="), EngineDir, false)
+				&& FParse::Value(FCommandLine::Get(), TEXT("platform="), Platform, false))
 			{
-				UE_LOG(LogPakFile, Error, TEXT("Expected 3 values, got %d, when parsing %s"), KeyValueText.Num(), *KeyFilename);
-				Pair.PrivateKey.Exponent.Zero();
+				static const TCHAR* SectionName = TEXT("Core.Encryption");
+
+				FConfigFile ConfigFile;
+				FConfigCacheIni::LoadExternalIniFile(ConfigFile, TEXT("Encryption"), *FPaths::Combine(EngineDir, TEXT("Config\\")), *FPaths::Combine(ProjectDir, TEXT("Config/")), true, *Platform);
+				bool bSignPak = false;
+				bool bEncryptPak = false;
+
+				ConfigFile.GetBool(SectionName, TEXT("SignPak"), bSignPak);
+				ConfigFile.GetBool(SectionName, TEXT("EncryptPak"), bEncryptPak);
+				
+				if (bSignPak)
+				{
+					FString RSAPublicExp, RSAPrivateExp, RSAModulus;
+					ConfigFile.GetString(SectionName, TEXT("rsa.publicexp"), RSAPublicExp);
+					ConfigFile.GetString(SectionName, TEXT("rsa.privateexp"), RSAPrivateExp);
+					ConfigFile.GetString(SectionName, TEXT("rsa.modulus"), RSAModulus);
+
+					OutEncryptionKeyPair.PrivateKey.Exponent.Parse(RSAPrivateExp);
+					OutEncryptionKeyPair.PrivateKey.Modulus.Parse(RSAModulus);
+					OutEncryptionKeyPair.PublicKey.Exponent.Parse(RSAPublicExp);
+					OutEncryptionKeyPair.PublicKey.Modulus = OutEncryptionKeyPair.PrivateKey.Modulus;
+
+					bSigningEnabled = true;
+
+					UE_LOG(LogPakFile, Display, TEXT("Parsed signature keys from config files."));
+				}
+
+				if (bEncryptPak)
+				{
+					ConfigFile.GetString(SectionName, TEXT("aes.key"), OutAESKey);
+
+					if (OutAESKey.Len() > 0)
+					{
+						UE_LOG(LogPakFile, Display, TEXT("Parsed AES encryption key from config files."));
+					}
+				}
 			}
 		}
-		else if (!ReadKeysFromFile(*KeyFilename, Pair))
+		else
 		{
-			UE_LOG(LogPakFile, Error, TEXT("Unable to load signature keys %s."), *KeyFilename);
+			FParse::Value(FCommandLine::Get(), TEXT("aes="), OutAESKey, false);
+
+			if (OutAESKey.Len() > 0)
+			{
+				UE_LOG(LogPakFile, Display, TEXT("Parsed AES encryption key from command line."));
+			}
+
+			if (FParse::Value(FCommandLine::Get(), TEXT("sign="), KeyFilename, false))
+			{
+				if (KeyFilename.StartsWith(TEXT("0x")))
+				{
+					TArray<FString> KeyValueText;
+					int32 NumParts = KeyFilename.ParseIntoArray(KeyValueText, TEXT("+"), true);
+					if (NumParts == 3)
+					{
+						OutEncryptionKeyPair.PrivateKey.Exponent.Parse(KeyValueText[0]);
+						OutEncryptionKeyPair.PrivateKey.Modulus.Parse(KeyValueText[1]);
+						OutEncryptionKeyPair.PublicKey.Exponent.Parse(KeyValueText[2]);
+						OutEncryptionKeyPair.PublicKey.Modulus = OutEncryptionKeyPair.PrivateKey.Modulus;
+
+						bSigningEnabled = true;
+
+						UE_LOG(LogPakFile, Display, TEXT("Parsed signature keys from command line."));
+					}
+					else
+					{
+						UE_LOG(LogPakFile, Error, TEXT("Expected 3, got %d, when parsing %s"), KeyValueText.Num(), *KeyFilename);
+						OutEncryptionKeyPair.PrivateKey.Exponent.Zero();
+					}
+				}
+				else if (!ReadKeysFromFile(*KeyFilename, OutEncryptionKeyPair))
+				{
+					UE_LOG(LogPakFile, Error, TEXT("Unable to load signature keys %s."), *KeyFilename);
+				}
+				else
+				{
+					bSigningEnabled = true;
+				}
+			}
 		}
 
-		if (!TestKeys(Pair))
+		if (OutAESKey.Len() > 0 && OutAESKey.Len() < 32)
 		{
-			Pair.PrivateKey.Exponent.Zero();
+			UE_LOG(LogPakFile, Fatal, TEXT("AES encryption key parsed from command line must be at least 32 characters long"));
 		}
 
-		if (!Pair.PrivateKey.Exponent.IsZero())
+		if (bSigningEnabled)
 		{
-			UE_LOG(LogPakFile, Display, TEXT("Creating signed pak %s."), Filename);
-			Writer = new FSignedArchiveWriter(*Writer, Filename, Pair.PublicKey, Pair.PrivateKey);
+			if (!OutEncryptionKeyPair.PrivateKey.Exponent.IsZero())
+			{
+				if (!TestKeys(OutEncryptionKeyPair))
+				{
+					OutEncryptionKeyPair.PrivateKey.Exponent.Zero();
+				}
+
+				UE_LOG(LogPakFile, Display, TEXT("Creating signed pak %s."), Filename);
+				Writer = new FSignedArchiveWriter(*Writer, Filename, OutEncryptionKeyPair.PublicKey, OutEncryptionKeyPair.PrivateKey);
+			}
+			else
+			{
+				UE_LOG(LogPakFile, Error, TEXT("Unable to create a signed pak writer."));
+				delete Writer;
+				Writer = NULL;
+			}
 		}
-		else
-		{
-			UE_LOG(LogPakFile, Error, TEXT("Unable to create a signed pak writer."));
-			delete Writer;
-			Writer = NULL;
-		}		
 	}
 	return Writer;
 }
@@ -763,13 +846,22 @@ bool CreatePakFile(const TCHAR* Filename, TArray<FPakInputPair>& FilesToAdd, con
 	const double StartTime = FPlatformTime::Seconds();
 
 	// Create Pak
-	TAutoPtr<FArchive> PakFileHandle(CreatePakWriter(Filename));
-	if (!PakFileHandle.IsValid())
+	FKeyPair KeyPair;
+	FString AESKey;
+	TUniquePtr<FArchive> PakFileHandle(CreatePakWriter(Filename, KeyPair, AESKey));
+	if (!PakFileHandle)
 	{
 		UE_LOG(LogPakFile, Error, TEXT("Unable to create pak file \"%s\"."), Filename);
 		return false;
 	}
 
+	ANSICHAR* EncryptionKeyAnsi = AESKey.Len() ? TCHAR_TO_ANSI(*AESKey) : nullptr;
+
+	if (EncryptionKeyAnsi)
+	{
+		UE_LOG(LogPakFile, Log, TEXT("Valid AES encryption key found - encryption will be enabled when requested!"));
+	}
+
 	FPakInfo Info;
 	TArray<FPakEntryPair> Index;
 	FString MountPoint = GetCommonRootPath(FilesToAdd);
@@ -799,6 +891,9 @@ bool CreatePakFile(const TCHAR* Filename, TArray<FPakInputPair>& FilesToAdd, con
 	uint64 TotalUncompressedSize = 0;
 	uint64 TotalCompressedSize = 0;
 
+	uint64 TotalRequestedEncryptedFiles = 0;
+	uint64 TotalEncryptedFiles = 0;
+
 	for (int32 FileIndex = 0; FileIndex < FilesToAdd.Num(); FileIndex++)
 	{
 		//  Remember the offset but don't serialize it with the entry header.
@@ -875,12 +970,12 @@ bool CreatePakFile(const TCHAR* Filename, TArray<FPakInputPair>& FilesToAdd, con
 		uint8* DataToWrite = nullptr;
 		if (FilesToAdd[FileIndex].bNeedsCompression && CompressionMethod != COMPRESS_None)
 		{
-			bCopiedToPak = PrepareCopyCompressedFileToPak(MountPoint, FilesToAdd[FileIndex], CompressedFileBuffer, NewEntry, DataToWrite, SizeToWrite);
+			bCopiedToPak = PrepareCopyCompressedFileToPak(MountPoint, FilesToAdd[FileIndex], CompressedFileBuffer, NewEntry, DataToWrite, SizeToWrite, EncryptionKeyAnsi);
 			DataToWrite = CompressedFileBuffer.CompressedBuffer.Get();
 		}
 		else
 		{
-			bCopiedToPak = PrepareCopyFileToPak(MountPoint, FilesToAdd[FileIndex], ReadBuffer, BufferSize, NewEntry, DataToWrite, SizeToWrite);
+			bCopiedToPak = PrepareCopyFileToPak(MountPoint, FilesToAdd[FileIndex], ReadBuffer, BufferSize, NewEntry, DataToWrite, SizeToWrite, EncryptionKeyAnsi);
 			DataToWrite = ReadBuffer;
 		}		
 
@@ -943,6 +1038,19 @@ bool CreatePakFile(const TCHAR* Filename, TArray<FPakInputPair>& FilesToAdd, con
 			// Update offset now and store it in the index (and only in index)
 			NewEntry.Info.Offset = NewEntryOffset;
 			Index.Add(NewEntry);
+			const TCHAR* EncryptedString = TEXT("");
+
+			if (FilesToAdd[FileIndex].bNeedEncryption)
+			{
+				TotalRequestedEncryptedFiles++;
+
+				if (EncryptionKeyAnsi != nullptr)
+				{
+					TotalEncryptedFiles++;
+					EncryptedString = TEXT("encrypted ");
+				}
+			}
+
 			if (FilesToAdd[FileIndex].bNeedsCompression && CompressionMethod != COMPRESS_None)
 			{
 				TotalCompressedSize += NewEntry.Info.Size;
@@ -950,22 +1058,22 @@ bool CreatePakFile(const TCHAR* Filename, TArray<FPakInputPair>& FilesToAdd, con
 				float PercentLess = ((float)NewEntry.Info.Size / (NewEntry.Info.UncompressedSize / 100.f));
 				if (FilesToAdd[FileIndex].SuggestedOrder < MAX_uint64)
 				{
-					UE_LOG(LogPakFile, Log, TEXT("Added compressed file \"%s\", %.2f%% of original size. Compressed Size %lld bytes, Original Size %lld bytes (order %llu)."), *NewEntry.Filename, PercentLess, NewEntry.Info.Size, NewEntry.Info.UncompressedSize, FilesToAdd[FileIndex].SuggestedOrder);
+					UE_LOG(LogPakFile, Log, TEXT("Added compressed %sfile \"%s\", %.2f%% of original size. Compressed Size %lld bytes, Original Size %lld bytes (order %llu)."), EncryptedString, *NewEntry.Filename, PercentLess, NewEntry.Info.Size, NewEntry.Info.UncompressedSize, FilesToAdd[FileIndex].SuggestedOrder);
 				}
 				else
 				{
-					UE_LOG(LogPakFile, Log, TEXT("Added compressed file \"%s\", %.2f%% of original size. Compressed Size %lld bytes, Original Size %lld bytes (no order given)."), *NewEntry.Filename, PercentLess, NewEntry.Info.Size, NewEntry.Info.UncompressedSize);
+					UE_LOG(LogPakFile, Log, TEXT("Added compressed %sfile \"%s\", %.2f%% of original size. Compressed Size %lld bytes, Original Size %lld bytes (no order given)."), EncryptedString, *NewEntry.Filename, PercentLess, NewEntry.Info.Size, NewEntry.Info.UncompressedSize);
 				}
 			}
 			else
 			{
 				if (FilesToAdd[FileIndex].SuggestedOrder < MAX_uint64)
 				{
-					UE_LOG(LogPakFile, Log, TEXT("Added file \"%s\", %lld bytes (order %llu)."), *NewEntry.Filename, NewEntry.Info.Size, FilesToAdd[FileIndex].SuggestedOrder);
+					UE_LOG(LogPakFile, Log, TEXT("Added %sfile \"%s\", %lld bytes (order %llu)."), EncryptedString, *NewEntry.Filename, NewEntry.Info.Size, FilesToAdd[FileIndex].SuggestedOrder);
 				}
 				else
 				{
-					UE_LOG(LogPakFile, Log, TEXT("Added file \"%s\", %lld bytes (no order given)."), *NewEntry.Filename, NewEntry.Info.Size);
+					UE_LOG(LogPakFile, Log, TEXT("Added %sfile \"%s\", %lld bytes (no order given)."), EncryptedString, *NewEntry.Filename, NewEntry.Info.Size);
 				}
 			}
 		}
@@ -1024,6 +1132,16 @@ bool CreatePakFile(const TCHAR* Filename, TArray<FPakInputPair>& FilesToAdd, con
 		UE_LOG(LogPakFile, Display, TEXT("Compression summary: %.2f%% of original size. Compressed Size %lld bytes, Original Size %lld bytes. "), PercentLess, TotalCompressedSize, TotalUncompressedSize);
 	}
 
+	if (TotalEncryptedFiles)
+	{
+		UE_LOG(LogPakFile, Display, TEXT("Encryption summary: %d files were encrypted. "), TotalEncryptedFiles);
+	}
+
+	if (TotalEncryptedFiles < TotalRequestedEncryptedFiles)
+	{
+		UE_LOG(LogPakFile, Warning, TEXT("%d files requested encryption, but no AES key was supplied! Encryption was skipped for these files"), TotalRequestedEncryptedFiles);
+	}
+
 	PakFileHandle->Close();
 	PakFileHandle.Reset();
 
@@ -1119,8 +1237,8 @@ bool ExtractFilesFromPak(const TCHAR* InPakFilename, const TCHAR* InDestPath, bo
 			{
 				FString DestFilename(DestPath / PakMountPoint /  It.Filename());
 
-				TAutoPtr<FArchive> FileHandle(IFileManager::Get().CreateFileWriter(*DestFilename));
-				if (FileHandle.IsValid())
+				TUniquePtr<FArchive> FileHandle(IFileManager::Get().CreateFileWriter(*DestFilename));
+				if (FileHandle)
 				{
 					if (Entry.CompressionMethod == COMPRESS_None)
 					{
@@ -1568,6 +1686,9 @@ INT32_MAIN_INT32_ARGC_TCHAR_ARGV()
 		UE_LOG(LogPakFile, Error, TEXT("    -encrypt"));
 		UE_LOG(LogPakFile, Error, TEXT("    -order=<OrderingFile>"));
 		UE_LOG(LogPakFile, Error, TEXT("    -diff (requires 2 filenames first)"));
+		UE_LOG(LogPakFile, Error, TEXT("    -enginedir (specify engine dir for when using ini encryption configs)"));
+		UE_LOG(LogPakFile, Error, TEXT("    -projectdir (specify project dir for when using ini encryption configs)"));
+		UE_LOG(LogPakFile, Error, TEXT("    -encryptionini (specify ini base name to gather encryption settings from)"));
 		return 1;
 	}
 
diff --git a/Engine/Source/Runtime/AIModule/Private/BrainComponent.cpp b/Engine/Source/Runtime/AIModule/Private/BrainComponent.cpp
--- a/Engine/Source/Runtime/AIModule/Private/BrainComponent.cpp
+++ b/Engine/Source/Runtime/AIModule/Private/BrainComponent.cpp
@@ -67,7 +67,7 @@ void FAIMessage::Broadcast(UObject* WorldContextObject, const FAIMessage& Messag
 	{
 		for (FConstControllerIterator It = MyWorld->GetControllerIterator(); It; ++It)
 		{
-			FAIMessage::Send(*It, Message);
+			FAIMessage::Send(It->Get(), Message);
 		}
 	}
 }
diff --git a/Engine/Source/Runtime/AIModule/Private/GameplayDebugger/GameplayDebuggerCategory_AI.cpp b/Engine/Source/Runtime/AIModule/Private/GameplayDebugger/GameplayDebuggerCategory_AI.cpp
--- a/Engine/Source/Runtime/AIModule/Private/GameplayDebugger/GameplayDebuggerCategory_AI.cpp
+++ b/Engine/Source/Runtime/AIModule/Private/GameplayDebugger/GameplayDebuggerCategory_AI.cpp
@@ -447,7 +447,7 @@ void FGameplayDebuggerCategory_AI::DrawPawnIcons(UWorld* World, AActor* DebugAct
 	FString FailsafeIcon = TEXT("/Engine/EngineResources/AICON-Green.AICON-Green");
 	for (FConstPawnIterator It = World->GetPawnIterator(); It; ++It)
 	{
-		const APawn* ItPawn = *It;
+		const APawn* ItPawn = It->Get();
 		if (IsValid(ItPawn) && SkipPawn != ItPawn)
 		{
 			const FVector IconLocation = ItPawn->GetActorLocation() + FVector(0, 0, ItPawn->GetSimpleCollisionHalfHeight());
diff --git a/Engine/Source/Runtime/AIModule/Private/Perception/AIPerceptionSystem.cpp b/Engine/Source/Runtime/AIModule/Private/Perception/AIPerceptionSystem.cpp
--- a/Engine/Source/Runtime/AIModule/Private/Perception/AIPerceptionSystem.cpp
+++ b/Engine/Source/Runtime/AIModule/Private/Perception/AIPerceptionSystem.cpp
@@ -555,7 +555,7 @@ void UAIPerceptionSystem::RegisterAllPawnsAsSourcesForSense(FAISenseID SenseID)
 	UWorld* World = GetWorld();
 	for (auto PawnIt = World->GetPawnIterator(); PawnIt; ++PawnIt)
 	{
-		if (*PawnIt)
+		if (PawnIt->Get())
 		{
 			RegisterSource(SenseID, **PawnIt);
 		}
diff --git a/Engine/Source/Runtime/AIModule/Private/Perception/PawnSensingComponent.cpp b/Engine/Source/Runtime/AIModule/Private/Perception/PawnSensingComponent.cpp
--- a/Engine/Source/Runtime/AIModule/Private/Perception/PawnSensingComponent.cpp
+++ b/Engine/Source/Runtime/AIModule/Private/Perception/PawnSensingComponent.cpp
@@ -207,7 +207,7 @@ void UPawnSensingComponent::UpdateAISensing()
 	{
 		for (FConstPlayerControllerIterator Iterator = Owner->GetWorld()->GetPlayerControllerIterator(); Iterator; ++Iterator)
 		{
-			APlayerController* PC = *Iterator;
+			APlayerController* PC = Iterator->Get();
 			if (IsValid(PC))
 			{
 				APawn* Pawn = PC->GetPawn();
@@ -222,7 +222,7 @@ void UPawnSensingComponent::UpdateAISensing()
 	{
 		for (FConstPawnIterator Iterator = Owner->GetWorld()->GetPawnIterator(); Iterator; ++Iterator)
 		{
-			APawn* Pawn = *Iterator;
+			APawn* Pawn = Iterator->Get();
 			if (IsValid(Pawn) && !IsSensorActor(Pawn))
 			{
 				SensePawn(*Pawn);
diff --git a/Engine/Source/Runtime/Core/Private/GenericPlatform/GenericPlatformCrashContext.cpp b/Engine/Source/Runtime/Core/Private/GenericPlatform/GenericPlatformCrashContext.cpp
--- a/Engine/Source/Runtime/Core/Private/GenericPlatform/GenericPlatformCrashContext.cpp
+++ b/Engine/Source/Runtime/Core/Private/GenericPlatform/GenericPlatformCrashContext.cpp
@@ -6,6 +6,10 @@
 #include "EngineVersion.h"
 #include "EngineBuildSettings.h"
 
+#ifndef NOINITCRASHREPORTER
+#define NOINITCRASHREPORTER 0
+#endif
+
 /*-----------------------------------------------------------------------------
 	FGenericCrashContext
 -----------------------------------------------------------------------------*/
@@ -73,6 +77,7 @@ namespace NCachedCrashContextProperties
 
 void FGenericCrashContext::Initialize()
 {
+#if !NOINITCRASHREPORTER
 	NCachedCrashContextProperties::bIsInternalBuild = FEngineBuildSettings::IsInternalBuild();
 	NCachedCrashContextProperties::bIsPerforceBuild = FEngineBuildSettings::IsPerforceBuild();
 	NCachedCrashContextProperties::bIsSourceDistribution = FEngineBuildSettings::IsSourceDistribution();
@@ -165,6 +170,7 @@ void FGenericCrashContext::Initialize()
 	FCoreDelegates::ConfigReadyForUse.AddStatic(FGenericCrashContext::InitializeFromConfig);
 
 	bIsInitialized = true;
+#endif
 }
 
 void FGenericCrashContext::InitializeFromConfig()
diff --git a/Engine/Source/Runtime/Core/Private/GenericPlatform/GenericPlatformFile.cpp b/Engine/Source/Runtime/Core/Private/GenericPlatform/GenericPlatformFile.cpp
--- a/Engine/Source/Runtime/Core/Private/GenericPlatform/GenericPlatformFile.cpp
+++ b/Engine/Source/Runtime/Core/Private/GenericPlatform/GenericPlatformFile.cpp
@@ -2,6 +2,7 @@
 
 #include "CorePrivatePCH.h"
 #include "ModuleManager.h"
+#include "UniquePtr.h"
 
 #if USE_NEW_ASYNC_IO
 
@@ -411,13 +412,13 @@ bool IPlatformFile::CopyFile(const TCHAR* To, const TCHAR* From, EPlatformFileRe
 {
 	const int64 MaxBufferSize = 1024*1024;
 
-	TAutoPtr<IFileHandle> FromFile(OpenRead(From, (ReadFlags & EPlatformFileRead::AllowWrite) != EPlatformFileRead::None));
-	if (!FromFile.IsValid())
+	TUniquePtr<IFileHandle> FromFile(OpenRead(From, (ReadFlags & EPlatformFileRead::AllowWrite) != EPlatformFileRead::None));
+	if (!FromFile)
 	{
 		return false;
 	}
-	TAutoPtr<IFileHandle> ToFile(OpenWrite(To, false, (WriteFlags & EPlatformFileWrite::AllowRead) != EPlatformFileWrite::None));
-	if (!ToFile.IsValid())
+	TUniquePtr<IFileHandle> ToFile(OpenWrite(To, false, (WriteFlags & EPlatformFileWrite::AllowRead) != EPlatformFileWrite::None));
+	if (!ToFile)
 	{
 		return false;
 	}
diff --git a/Engine/Source/Runtime/Core/Private/GenericPlatform/GenericPlatformOutputDevices.cpp b/Engine/Source/Runtime/Core/Private/GenericPlatform/GenericPlatformOutputDevices.cpp
--- a/Engine/Source/Runtime/Core/Private/GenericPlatform/GenericPlatformOutputDevices.cpp
+++ b/Engine/Source/Runtime/Core/Private/GenericPlatform/GenericPlatformOutputDevices.cpp
@@ -5,6 +5,7 @@
 #include "Misc/App.h"
 #include "Misc/OutputDeviceConsole.h"
 #include "Misc/OutputDeviceMemory.h"
+#include "UniquePtr.h"
 
 void FGenericPlatformOutputDevices::SetupOutputDevices()
 {
@@ -84,12 +85,12 @@ class FOutputDevice* FGenericPlatformOutputDevices::GetLog()
 {
 	static struct FLogOutputDeviceInitializer
 	{
-		TAutoPtr<FOutputDevice> LogDevice;
+		TUniquePtr<FOutputDevice> LogDevice;
 		FLogOutputDeviceInitializer()
 		{
 #if WITH_LOGGING_TO_MEMORY
 #if !IS_PROGRAM && !WITH_EDITORONLY_DATA
-			if (!LogDevice.IsValid() 
+			if (!LogDevice
 #if !(UE_BUILD_SHIPPING || UE_BUILD_TEST)
 				 && FParse::Param(FCommandLine::Get(), TEXT("LOGTOMEMORY")) 
 #else
@@ -97,19 +98,19 @@ class FOutputDevice* FGenericPlatformOutputDevices::GetLog()
 #endif
 				 )
 			{
-				LogDevice = new FOutputDeviceMemory();
+				LogDevice = MakeUnique<FOutputDeviceMemory>();
 			}
 #endif // !IS_PROGRAM && !WITH_EDITORONLY_DATA
 #endif // WITH_LOGGING_TO_MEMORY
-			if (!LogDevice.IsValid())
+			if (!LogDevice)
 			{
-				LogDevice = new FOutputDeviceFile();
+				LogDevice = MakeUnique<FOutputDeviceFile>();
 			}
 		}
 
 	} Singleton;
 
-	return Singleton.LogDevice.GetOwnedPointer();
+	return Singleton.LogDevice.Get();
 }
 
 
diff --git a/Engine/Source/Runtime/Core/Private/HAL/FileManagerGeneric.cpp b/Engine/Source/Runtime/Core/Private/HAL/FileManagerGeneric.cpp
--- a/Engine/Source/Runtime/Core/Private/HAL/FileManagerGeneric.cpp
+++ b/Engine/Source/Runtime/Core/Private/HAL/FileManagerGeneric.cpp
@@ -13,6 +13,7 @@
 #include "CorePrivatePCH.h"
 #include "FileManagerGeneric.h"
 #include "SecureHash.h"
+#include "UniquePtr.h"
 #include <time.h>
 
 DEFINE_LOG_CATEGORY_STATIC( LogFileManager, Log, All );
@@ -864,10 +865,10 @@ void FArchiveFileWriterGeneric::LogWriteError(const TCHAR* Message)
 
 IFileManager& IFileManager::Get()
 {
-	static TScopedPointer<FFileManagerGeneric> AutoDestroySingleton;
+	static TUniquePtr<FFileManagerGeneric> AutoDestroySingleton;
 	if( !AutoDestroySingleton )
 	{
-		AutoDestroySingleton = new FFileManagerGeneric();
+		AutoDestroySingleton = MakeUnique<FFileManagerGeneric>();
 	}
 	return *AutoDestroySingleton;
 }
