Set the correct exit code for UnrealPak. Should return 0 on success, not 1.

#rnx

diff --git a/Engine/Source/Programs/UnrealPak/Private/UnrealPak.cpp b/Engine/Source/Programs/UnrealPak/Private/UnrealPak.cpp
--- a/Engine/Source/Programs/UnrealPak/Private/UnrealPak.cpp
+++ b/Engine/Source/Programs/UnrealPak/Private/UnrealPak.cpp
@@ -14,7 +14,7 @@ INT32_MAIN_INT32_ARGC_TCHAR_ARGV()
 
 	double StartTime = FPlatformTime::Seconds();
 
-	int32 Result = ExecuteUnrealPak(FCommandLine::Get());
+	int32 Result = ExecuteUnrealPak(FCommandLine::Get())? 0 : 1;
 
 	UE_LOG(LogPakFile, Display, TEXT("Unreal pak executed in %f seconds"), FPlatformTime::Seconds() - StartTime );
 
