Include plugins for target platforms when building UnrealPak.

#rnx

diff --git a/Engine/Source/Programs/UnrealPak/UnrealPak.Target.cs b/Engine/Source/Programs/UnrealPak/UnrealPak.Target.cs
--- a/Engine/Source/Programs/UnrealPak/UnrealPak.Target.cs
+++ b/Engine/Source/Programs/UnrealPak/UnrealPak.Target.cs
@@ -14,6 +14,7 @@ public class UnrealPakTarget : TargetRules
 		bCompileLeanAndMeanUE = true;
 		bUseMallocProfiler = false;
 		bCompileWithPluginSupport = true;
+		bIncludePluginsForTargetPlatforms = true;
 
 		// No editor needed
 		bBuildEditor = false;
