Enable trace for the UnrealPak program

#rnx

diff --git a/Engine/Source/Programs/UnrealPak/UnrealPak.Target.cs b/Engine/Source/Programs/UnrealPak/UnrealPak.Target.cs
--- a/Engine/Source/Programs/UnrealPak/UnrealPak.Target.cs
+++ b/Engine/Source/Programs/UnrealPak/UnrealPak.Target.cs
@@ -29,5 +29,7 @@ public class UnrealPakTarget : TargetRules
 		// UnrealPak is a console application, not a Windows app (sets entry point to main(), instead of WinMain())
 		bCompileAgainstApplicationCore = true;
 		bIsBuildingConsoleApplication = true;
+
+		GlobalDefinitions.Add("UE_TRACE_ENABLED=1");
 	}
 }
