Fix for issue where running UnrealPak with -WriteBackMetadataToAssetRegistry=OriginalFile would crash because UObjectRedirector was not properly initialized

diff --git a/Engine/Source/Programs/UnrealPak/UnrealPak.Build.cs b/Engine/Source/Programs/UnrealPak/UnrealPak.Build.cs
--- a/Engine/Source/Programs/UnrealPak/UnrealPak.Build.cs
+++ b/Engine/Source/Programs/UnrealPak/UnrealPak.Build.cs
@@ -8,7 +8,7 @@ public class UnrealPak : ModuleRules
 	{
 		PublicIncludePaths.Add("Runtime/Launch/Public");
 
-		PrivateDependencyModuleNames.AddRange(new string[] { "Core", "PakFile", "Json", "Projects", "PakFileUtilities", "RSA", "ApplicationCore" });
+		PrivateDependencyModuleNames.AddRange(new string[] { "Core", "CoreUObject", "AssetRegistry", "PakFile", "Json", "Projects", "PakFileUtilities", "RSA", "ApplicationCore" });
 
 		PrivateIncludePaths.Add("Runtime/Launch/Private");      // For LaunchEngineLoop.cpp include
 
diff --git a/Engine/Source/Programs/UnrealPak/UnrealPak.Target.cs b/Engine/Source/Programs/UnrealPak/UnrealPak.Target.cs
--- a/Engine/Source/Programs/UnrealPak/UnrealPak.Target.cs
+++ b/Engine/Source/Programs/UnrealPak/UnrealPak.Target.cs
@@ -20,9 +20,9 @@ public class UnrealPakTarget : TargetRules
 		// Editor-only data, however, is needed
 		bBuildWithEditorOnlyData = true;
 
-		// Currently this app is not linking against the engine, so we'll compile out references from Core to the rest of the engine
+		// Currently this app uses AssetRegistry and CoreUObject, but not Engine
 		bCompileAgainstEngine = false;
-		bCompileAgainstCoreUObject = false;
+		bCompileAgainstCoreUObject = true;
 
 		// ICU is not needed
 		bCompileICU = false;
