UnrealPak: Optimizations for small workloads.
Add -fastexit parameter to terminate process instead of shutting down engine.
Turn off build target features to build less code (total dll size reduced from 56 MB to 46 MB).
Disable external profiler supports (skips loading of some dlls).
Don't spawn any reserve workers (by adding a DefaultEngine.ini for UnrealPak).
Make UnrealPAk compile in Test config (by fixing compile error in PerforceConnection.cpp).

#rnx

diff --git a/Engine/Plugins/Developer/PerforceSourceControl/Source/PerforceSourceControl/Private/PerforceConnection.cpp b/Engine/Plugins/Developer/PerforceSourceControl/Source/PerforceSourceControl/Private/PerforceConnection.cpp
--- a/Engine/Plugins/Developer/PerforceSourceControl/Source/PerforceSourceControl/Private/PerforceConnection.cpp
+++ b/Engine/Plugins/Developer/PerforceSourceControl/Source/PerforceSourceControl/Private/PerforceConnection.cpp
@@ -1071,7 +1071,14 @@ bool FPerforceConnection::RunCommand(	const FString& InCommand, const TArray<FSt
 	if (bLogCommandDetails)
 	{
 		const double ExecutionTime = FPlatformTime::Seconds() - SCCStartTime;
-		UE_CLOG((LogSourceControl.GetVerbosity() == ELogVerbosity::VeryVerbose) || (ExecutionTime >= 0.1), LogSourceControl, Log, TEXT("P4 execution time: %0.4f seconds. Command: %s"), ExecutionTime, *FullCommand);
+		if (ExecutionTime >= 0.1)
+		{
+			UE_LOG(LogSourceControl, Log, TEXT("P4 execution time: %0.4f seconds. Command: %s"), ExecutionTime, *FullCommand);
+		}
+		else
+		{
+			UE_LOG(LogSourceControl, VeryVerbose, TEXT("P4 execution time: %0.4f seconds. Command: %s"), ExecutionTime, *FullCommand);
+		}
 	}
 
 	return OutRecordSet.Num() > 0;
diff --git a/Engine/Programs/UnrealPak/Config/DefaultEngine.ini b/Engine/Programs/UnrealPak/Config/DefaultEngine.ini
new file mode 100644
--- /dev/null
+++ b/Engine/Programs/UnrealPak/Config/DefaultEngine.ini
@@ -0,0 +1,2 @@
+[TaskGraph]
+EnableReserveWorkers=false
diff --git a/Engine/Source/Programs/UnrealPak/Private/UnrealPak.cpp b/Engine/Source/Programs/UnrealPak/Private/UnrealPak.cpp
--- a/Engine/Source/Programs/UnrealPak/Private/UnrealPak.cpp
+++ b/Engine/Source/Programs/UnrealPak/Private/UnrealPak.cpp
@@ -20,6 +20,11 @@ INT32_MAIN_INT32_ARGC_TCHAR_ARGV()
 
 	UE_LOG(LogPakFile, Display, TEXT("UnrealPak executed in %f seconds"), FPlatformTime::Seconds() - StartTime);
 
+	if (FParse::Param(FCommandLine::Get(), TEXT("fastexit")))
+	{
+		FPlatformMisc::RequestExitWithStatus(true, Result);
+	}
+
 	GLog->Flush();
 
 	RequestEngineExit(TEXT("UnrealPak Exiting"));
diff --git a/Engine/Source/Programs/UnrealPak/UnrealPak.Target.cs b/Engine/Source/Programs/UnrealPak/UnrealPak.Target.cs
--- a/Engine/Source/Programs/UnrealPak/UnrealPak.Target.cs
+++ b/Engine/Source/Programs/UnrealPak/UnrealPak.Target.cs
@@ -12,25 +12,34 @@ public class UnrealPakTarget : TargetRules
 		IncludeOrderVersion = EngineIncludeOrderVersion.Latest;
 		LaunchModuleName = "UnrealPak";
 
-		bBuildDeveloperTools = false;
-		bUseMallocProfiler = false;
-		bCompileWithPluginSupport = true;
-		bIncludePluginsForTargetPlatforms = true;
-
-		// Editor-only data, however, is needed
-		bBuildWithEditorOnlyData = true;
-
-		// Currently this app uses AssetRegistry and CoreUObject, but not Engine
-		bCompileAgainstEngine = false;
-		bCompileAgainstCoreUObject = true;
-
-		// ICU is not needed
-		bCompileICU = false;
-
-		// UnrealPak is a console application, not a Windows app (sets entry point to main(), instead of WinMain())
-		bCompileAgainstApplicationCore = true;
-		bIsBuildingConsoleApplication = true;
+		bUseXGEController					= false;
+		bCompileFreeType					= false;
+		bLoggingToMemoryEnabled				= false;
+		bUseLoggingInShipping				= true;
+		bCompileWithAccessibilitySupport	= false;
+		bCompileWithPluginSupport			= true;
+		bIncludePluginsForTargetPlatforms	= true;
+		bWithServerCode						= false;
+		bCompileNavmeshClusterLinks			= false;
+		bCompileNavmeshSegmentLinks			= false;
+		bCompileRecast						= false;
+		bCompileICU 						= false;
+		bWithLiveCoding						= false;
+		bBuildDeveloperTools				= false;
+		bUseMallocProfiler					= false;
+		bBuildWithEditorOnlyData			= true;
+		bCompileAgainstEngine				= false;
+		bCompileAgainstCoreUObject			= true;
+		bCompileAgainstApplicationCore		= true;
+		bUsesSlate							= false;
+		bIsBuildingConsoleApplication		= true;
+
+		WindowsPlatform.bUseBundledDbgHelp = false;
+		WindowsPlatform.bPixProfilingEnabled = false;
+		WindowsPlatform.bMergeIdenticalCOMDATs = true;
+		WindowsPlatform.bStripUnreferencedSymbols = true;
 
 		GlobalDefinitions.Add("UE_TRACE_ENABLED=1");
+		GlobalDefinitions.Add("UE_EXTERNAL_PROFILING_ENABLED=0");
 	}
 }
