IoStoreUtilities: Add TargetPlatform module to UnrealPak so that GenerateZenFileSystemManifest and other commands needing TargetPlatform will work.
#rnx

diff --git a/Engine/Source/Programs/UnrealPak/UnrealPak.Build.cs b/Engine/Source/Programs/UnrealPak/UnrealPak.Build.cs
--- a/Engine/Source/Programs/UnrealPak/UnrealPak.Build.cs
+++ b/Engine/Source/Programs/UnrealPak/UnrealPak.Build.cs
@@ -8,16 +8,17 @@ public class UnrealPak : ModuleRules
 	{
 		PublicIncludePathModuleNames.Add("Launch");
 
-		PrivateDependencyModuleNames.AddRange(new string[] { 
+		PrivateDependencyModuleNames.AddRange(new string[] {
+			"ApplicationCore",
+			"AssetRegistry",
 			"Core", 
-			"CoreUObject", 
-			"AssetRegistry", 
-			"PakFile", 
-			"Json", 
+			"CoreUObject",
+			"Json",
+			"PakFile",
+			"PakFileUtilities",
 			"Projects", 
-			"PakFileUtilities", 
 			"RSA", 
-			"ApplicationCore" 
+			"TargetPlatform",
 		});
 
 		if (Target.bBuildWithEditorOnlyData)
diff --git a/Engine/Source/Programs/UnrealPak/UnrealPak.Target.cs b/Engine/Source/Programs/UnrealPak/UnrealPak.Target.cs
--- a/Engine/Source/Programs/UnrealPak/UnrealPak.Target.cs
+++ b/Engine/Source/Programs/UnrealPak/UnrealPak.Target.cs
@@ -32,6 +32,8 @@ public class UnrealPakTarget : TargetRules
 		bCompileAgainstApplicationCore		= true;
 		bUsesSlate							= false;
 		bIsBuildingConsoleApplication		= true;
+		bForceBuildTargetPlatforms			= true;
+		bNeedsExtraShaderFormats			= true;
 
 		WindowsPlatform.bUseBundledDbgHelp = false;
 		WindowsPlatform.bPixProfilingEnabled = false;
